Motion Design: added data link and scene state integration plugins

diff --git a/Engine/Plugins/Experimental/AvalancheDataLink/AvalancheDataLink.uplugin b/Engine/Plugins/Experimental/AvalancheDataLink/AvalancheDataLink.uplugin
new file mode 100644
--- /dev/null
+++ b/Engine/Plugins/Experimental/AvalancheDataLink/AvalancheDataLink.uplugin
@@ -0,0 +1,38 @@
+{
+	"FileVersion": 3,
+	"Version": 1,
+	"VersionName": "1.0",
+	"FriendlyName": "DataLink for Motion Design",
+	"Description": "",
+	"Category": "Other",
+	"CreatedBy": "",
+	"CreatedByURL": "",
+	"DocsURL": "",
+	"MarketplaceURL": "",
+	"CanContainContent": true,
+	"IsBetaVersion": false,
+	"IsExperimentalVersion": true,
+	"Installed": false,
+	"Modules": [
+		{
+			"Name": "AvalancheDataLink",
+			"Type": "Runtime",
+			"LoadingPhase": "Default"
+		},
+		{
+			"Name": "AvalancheDataLinkEditor",
+			"Type": "Editor",
+			"LoadingPhase": "Default"
+		}
+	],
+	"Plugins": [
+		{
+			"Name": "Avalanche",
+			"Enabled": true
+		},
+		{
+			"Name": "DataLink",
+			"Enabled": true
+		}
+	]
+}
\ No newline at end of file
diff --git a/Engine/Plugins/Experimental/AvalancheDataLink/Source/AvalancheDataLink/AvalancheDataLink.Build.cs b/Engine/Plugins/Experimental/AvalancheDataLink/Source/AvalancheDataLink/AvalancheDataLink.Build.cs
new file mode 100644
--- /dev/null
+++ b/Engine/Plugins/Experimental/AvalancheDataLink/Source/AvalancheDataLink/AvalancheDataLink.Build.cs
@@ -0,0 +1,34 @@
+// Copyright Epic Games, Inc. All Rights Reserved.
+
+using UnrealBuildTool;
+
+public class AvalancheDataLink : ModuleRules
+{
+    public AvalancheDataLink(ReadOnlyTargetRules Target) : base(Target)
+    {
+        PCHUsage = ModuleRules.PCHUsageMode.UseExplicitOrSharedPCHs;
+
+        PublicDependencyModuleNames.AddRange(
+            new string[]
+            {
+                "AvalancheRemoteControl",
+                "Core",
+                "CoreUObject",
+                "DataLink",
+            }
+        );
+
+        PrivateDependencyModuleNames.AddRange(
+            new string[]
+            {
+                "Avalanche",
+                "DataLinkJson",
+                "Engine",
+                "Json",
+                "JsonUtilities",
+                "RemoteControl",
+                "RemoteControlLogic",
+            }
+        );
+    }
+}
diff --git a/Engine/Plugins/Experimental/AvalancheDataLink/Source/AvalancheDataLink/Private/AvaDataLinkActor.cpp b/Engine/Plugins/Experimental/AvalancheDataLink/Source/AvalancheDataLink/Private/AvaDataLinkActor.cpp
new file mode 100644
--- /dev/null
+++ b/Engine/Plugins/Experimental/AvalancheDataLink/Source/AvalancheDataLink/Private/AvaDataLinkActor.cpp
@@ -0,0 +1,15 @@
+// Copyright Epic Games, Inc. All Rights Reserved.
+
+#include "AvaDataLinkActor.h"
+#include "AvaDataLinkInstance.h"
+
+void AAvaDataLinkActor::ExecuteDataLinkInstances()
+{
+	for (UAvaDataLinkInstance* DataLinkInstance: DataLinkInstances)
+	{
+		if (DataLinkInstance)
+		{
+			DataLinkInstance->Execute();
+		}
+	}
+}
diff --git a/Engine/Plugins/Experimental/AvalancheDataLink/Source/AvalancheDataLink/Private/AvaDataLinkActor.h b/Engine/Plugins/Experimental/AvalancheDataLink/Source/AvalancheDataLink/Private/AvaDataLinkActor.h
new file mode 100644
--- /dev/null
+++ b/Engine/Plugins/Experimental/AvalancheDataLink/Source/AvalancheDataLink/Private/AvaDataLinkActor.h
@@ -0,0 +1,21 @@
+// Copyright Epic Games, Inc. All Rights Reserved.
+
+#pragma once
+
+#include "GameFramework/Actor.h"
+#include "AvaDataLinkActor.generated.h"
+
+class UAvaDataLinkInstance;
+
+UCLASS(DisplayName="Motion Design Data Link Actor", HideCategories=(Activation, Actor, AssetUserData, Collision, Cooking, DataLayers, HLOD, Input, LevelInstance, Mobility, Navigation, Networking, Physics, Rendering, Replication, Tags, Transform, WorldPartition))
+class AAvaDataLinkActor : public AActor
+{
+	GENERATED_BODY()
+
+public:
+	UFUNCTION(BlueprintCallable, CallInEditor, Category="Motion Design Data Link")
+	void ExecuteDataLinkInstances();
+
+	UPROPERTY(EditAnywhere, Instanced, Category="Motion Design Data Link")
+	TArray<TObjectPtr<UAvaDataLinkInstance>> DataLinkInstances;
+};
diff --git a/Engine/Plugins/Experimental/AvalancheDataLink/Source/AvalancheDataLink/Private/AvaDataLinkInstance.cpp b/Engine/Plugins/Experimental/AvalancheDataLink/Source/AvalancheDataLink/Private/AvaDataLinkInstance.cpp
new file mode 100644
--- /dev/null
+++ b/Engine/Plugins/Experimental/AvalancheDataLink/Source/AvalancheDataLink/Private/AvaDataLinkInstance.cpp
@@ -0,0 +1,202 @@
+// Copyright Epic Games, Inc. All Rights Reserved.
+
+#include "AvaDataLinkInstance.h"
+#include "AvaDataLinkLog.h"
+#include "AvaSceneSubsystem.h"
+#include "Avalanche/Public/IAvaSceneInterface.h"
+#include "Controller/RCController.h"
+#include "DataLinkEnums.h"
+#include "DataLinkExecutor.h"
+#include "DataLinkExecutorArguments.h"
+#include "DataLinkJsonUtils.h"
+#include "DataLinkUtils.h"
+#include "Engine/Level.h"
+#include "Engine/World.h"
+#include "JsonObjectConverter.h"
+#include "JsonObjectWrapper.h"
+#include "RCVirtualProperty.h"
+#include "RCVirtualPropertyContainer.h"
+#include "RemoteControlPreset.h"
+#include "StructUtils/StructView.h"
+#include "UObject/PropertyAccessUtil.h"
+
+void UAvaDataLinkInstance::Execute()
+{
+	if (Executor.IsValid())
+	{
+		UE_LOG(LogAvaDataLink, Error, TEXT("[%s] Data Link execution is in progress!"), Executor->GetContextName().GetData());
+		return;
+	}
+
+	Executor = FDataLinkExecutor::Create(FDataLinkExecutorArguments(DataLinkInstance)
+#if WITH_DATALINK_CONTEXT
+		.SetContextName(BuildContextName())
+#endif
+		.SetContextObject(this)
+		.SetOnFinish(FOnDataLinkExecutionFinished::CreateUObject(this, &UAvaDataLinkInstance::OnExecutionFinished)));
+
+	Executor->Run();
+}
+
+#if WITH_DATALINK_CONTEXT
+FString UAvaDataLinkInstance::BuildContextName() const
+{
+	return FString::Printf(TEXT("Motion Design Data Link. World: '%s'"), *GetNameSafe(GetTypedOuter<UWorld>()));
+}
+#endif
+
+void UAvaDataLinkInstance::OnExecutionFinished(const FDataLinkExecutor& InExecutor, EDataLinkExecutionResult InResult, FConstStructView InOutputDataView)
+{
+	Executor.Reset();
+
+	if (InResult == EDataLinkExecutionResult::Failed || !InOutputDataView.IsValid())
+	{
+		return;
+	}
+
+	URemoteControlPreset* const Preset = GetRemoteControlPreset();
+	if (!Preset)
+	{
+		UE_LOG(LogAvaDataLink, Error, TEXT("[%s] Data Link execution finished, but Remote Control is invalid!"), InExecutor.GetContextName().GetData());
+		return;
+	}
+
+	URCVirtualPropertyContainerBase* const ControllerContainer = Preset->GetControllerContainer();
+	if (!ControllerContainer)
+	{
+		UE_LOG(LogAvaDataLink, Error, TEXT("[%s] Data Link execution finished, but Remote Control '%s' has invalid Controller Container!")
+			, InExecutor.GetContextName().GetData()
+			, *Preset->GetName());
+		return;
+	}
+
+	const FStructView TargetDataView = ControllerContainer->GetPropertyBagMutableValue();
+	if (!TargetDataView.IsValid())
+	{
+		UE_LOG(LogAvaDataLink, Error, TEXT("[%s] Data Link execution finished, but Remote Control '%s' has invalid Controller Container Data View!")
+			, InExecutor.GetContextName().GetData()
+			, *Preset->GetName());
+		return;
+	}
+
+	TArray<URCController*> ModifiedControllers;
+	ModifiedControllers.Reserve(ControllerMappings.Num());
+
+	if (const FJsonObjectWrapper* OutputJson = InOutputDataView.GetPtr<const FJsonObjectWrapper>())
+	{
+		if (!OutputJson->JsonObject.IsValid())
+		{
+			UE_LOG(LogAvaDataLink, Warning, TEXT("[%s] Data Link output could not be applied to controllers. Json Object was not valid!"), InExecutor.GetContextName().GetData());
+			return;
+		}
+
+		TSharedRef<FJsonObject> JsonObject = OutputJson->JsonObject.ToSharedRef();
+
+		ForEachResolvedController(InExecutor, Preset, TargetDataView,
+			[&JsonObject, &ModifiedControllers](const FResolvedController& InResolvedController)
+			{
+				TSharedPtr<FJsonValue> SourceJsonValue = UE::DataLinkJson::FindJsonValue(JsonObject, InResolvedController.Mapping->OutputFieldName);
+
+				// If property was set, return true to signal that the controller value was modified
+				if (FJsonObjectConverter::JsonValueToUProperty(SourceJsonValue, InResolvedController.TargetProperty, InResolvedController.TargetMemory))
+				{
+					ModifiedControllers.AddUnique(InResolvedController.Controller);
+				}
+			});
+	}
+	else
+	{
+		ForEachResolvedController(InExecutor, Preset, TargetDataView,
+			[&InOutputDataView, &ModifiedControllers](const FResolvedController& InResolvedController)
+			{
+				const FProperty* SourceProperty = PropertyAccessUtil::FindPropertyByName(*InResolvedController.Mapping->OutputFieldName, InOutputDataView.GetScriptStruct());
+				if (!SourceProperty)
+				{
+					return;
+				}
+
+				if (!SourceProperty->SameType(InResolvedController.TargetProperty))
+				{
+					return;
+				}
+
+				const uint8* SourceMemory = SourceProperty->ContainerPtrToValuePtr<uint8>(InOutputDataView.GetMemory());
+				InResolvedController.TargetProperty->CopyCompleteValue(InResolvedController.TargetMemory, SourceMemory);
+				ModifiedControllers.AddUnique(InResolvedController.Controller);
+			});
+	}
+
+	TSet<FGuid> ModifiedControllerIds;
+	ModifiedControllerIds.Reserve(ModifiedControllers.Num());
+
+	for (URCController* Controller : ModifiedControllers)
+	{
+		Controller->OnModifyPropertyValue();
+		ModifiedControllerIds.Add(Controller->Id);
+	}
+
+	Preset->OnControllerModified().Broadcast(Preset, ModifiedControllerIds);
+}
+
+IAvaSceneInterface* UAvaDataLinkInstance::GetSceneInterface() const
+{
+	ULevel* Level = GetTypedOuter<ULevel>();
+	if (!Level || !Level->OwningWorld)
+	{
+		return nullptr;
+	}
+
+	const UAvaSceneSubsystem* SceneSubsystem = Level->OwningWorld->GetSubsystem<UAvaSceneSubsystem>();
+	if (!SceneSubsystem)
+	{
+		// No subsystem yet available, try finding the scene interface by iterating the actors
+		return UAvaSceneSubsystem::FindSceneInterface(Level);
+	}
+
+	return SceneSubsystem->GetSceneInterface(Level);
+}
+
+URemoteControlPreset* UAvaDataLinkInstance::GetRemoteControlPreset() const
+{
+	if (IAvaSceneInterface* SceneInterface = GetSceneInterface())
+	{
+		return SceneInterface->GetRemoteControlPreset();
+	}
+	return nullptr;
+}
+
+void UAvaDataLinkInstance::ForEachResolvedController(const FDataLinkExecutor& InExecutor, URemoteControlPreset* InPreset, FStructView InTargetDataView, TFunctionRef<void(const FResolvedController&)> InFunction)
+{
+	for (const FAvaDataLinkControllerMapping& Mapping : ControllerMappings)
+	{
+		FResolvedController ResolvedMapping;
+		ResolvedMapping.Mapping = &Mapping;
+
+		// Get the controller to retrieve the underlying Property FName of the controller within the Controller Container Property Bag
+		ResolvedMapping.Controller = Cast<URCController>(Mapping.TargetController.FindController(InPreset));
+		if (!ResolvedMapping.Controller)
+		{
+			UE_LOG(LogAvaDataLink, Warning, TEXT("[%s] Data Link output '%s' could not be applied to controller '%s'. Controller was not found in preset '%s'!")
+				, InExecutor.GetContextName().GetData()
+				, *Mapping.OutputFieldName
+				, *Mapping.TargetController.Name.ToString()
+				, *InPreset->GetName());
+			continue;
+		}
+
+		ResolvedMapping.TargetProperty = PropertyAccessUtil::FindPropertyByName(ResolvedMapping.Controller->PropertyName, InTargetDataView.GetScriptStruct());
+		if (!ResolvedMapping.TargetProperty)
+		{
+			UE_LOG(LogAvaDataLink, Warning, TEXT("[%s] Data Link output '%s' could not be applied to controller '%s'. Controller property '%s' in preset '%s' was not found!")
+				, InExecutor.GetContextName().GetData()
+				, *Mapping.OutputFieldName
+				, *Mapping.TargetController.Name.ToString()
+				, *ResolvedMapping.Controller->PropertyName.ToString()
+				, *InPreset->GetName());
+			continue;
+		}
+
+		ResolvedMapping.TargetMemory = ResolvedMapping.TargetProperty->ContainerPtrToValuePtr<uint8>(InTargetDataView.GetMemory());
+		InFunction(ResolvedMapping);
+	}
+}
diff --git a/Engine/Plugins/Experimental/AvalancheDataLink/Source/AvalancheDataLink/Private/AvaDataLinkLog.h b/Engine/Plugins/Experimental/AvalancheDataLink/Source/AvalancheDataLink/Private/AvaDataLinkLog.h
new file mode 100644
--- /dev/null
+++ b/Engine/Plugins/Experimental/AvalancheDataLink/Source/AvalancheDataLink/Private/AvaDataLinkLog.h
@@ -0,0 +1,7 @@
+// Copyright Epic Games, Inc. All Rights Reserved.
+
+#pragma once
+
+#include "Logging/LogMacros.h"
+
+DECLARE_LOG_CATEGORY_EXTERN(LogAvaDataLink, Log, All);
diff --git a/Engine/Plugins/Experimental/AvalancheDataLink/Source/AvalancheDataLink/Private/AvaDataLinkModule.cpp b/Engine/Plugins/Experimental/AvalancheDataLink/Source/AvalancheDataLink/Private/AvaDataLinkModule.cpp
new file mode 100644
--- /dev/null
+++ b/Engine/Plugins/Experimental/AvalancheDataLink/Source/AvalancheDataLink/Private/AvaDataLinkModule.cpp
@@ -0,0 +1,8 @@
+// Copyright Epic Games, Inc. All Rights Reserved.
+
+#include "AvaDataLinkLog.h"
+#include "Modules/ModuleManager.h"
+
+DEFINE_LOG_CATEGORY(LogAvaDataLink);
+
+IMPLEMENT_MODULE(FDefaultModuleImpl, AvalancheDataLink)
diff --git a/Engine/Plugins/Experimental/AvalancheDataLink/Source/AvalancheDataLink/Public/AvaDataLinkInstance.h b/Engine/Plugins/Experimental/AvalancheDataLink/Source/AvalancheDataLink/Public/AvaDataLinkInstance.h
new file mode 100644
--- /dev/null
+++ b/Engine/Plugins/Experimental/AvalancheDataLink/Source/AvalancheDataLink/Public/AvaDataLinkInstance.h
@@ -0,0 +1,74 @@
+// Copyright Epic Games, Inc. All Rights Reserved.
+
+#pragma once
+
+#include "AvaRCControllerId.h"
+#include "DataLinkInstance.h"
+#include "UObject/Object.h"
+#include "AvaDataLinkInstance.generated.h"
+
+class FDataLinkExecutor;
+class IAvaSceneInterface;
+class URCController;
+enum class EDataLinkExecutionResult : uint8;
+struct FConstStructView;
+struct FStructView;
+
+USTRUCT(BlueprintType, DisplayName="Motion Design Data Link RC Controller Mapping")
+struct FAvaDataLinkControllerMapping
+{
+	GENERATED_BODY()
+
+	UPROPERTY(EditAnywhere, BlueprintReadWrite, Category = "Motion Design Data Link")
+	FString OutputFieldName;
+
+	UPROPERTY(EditAnywhere, BlueprintReadWrite, Category = "Motion Design Data Link")
+	FAvaRCControllerId TargetController;
+};
+
+UCLASS(MinimalAPI, EditInlineNew)
+class UAvaDataLinkInstance : public UObject
+{
+	GENERATED_BODY()
+
+public:
+	void Execute();
+
+	static FName GetDataLinkInstancePropertyName()
+	{
+		return GET_MEMBER_NAME_CHECKED(UAvaDataLinkInstance, DataLinkInstance);
+	}
+
+	static FName GetControllerMappingsPropertyName()
+	{
+		return GET_MEMBER_NAME_CHECKED(UAvaDataLinkInstance, ControllerMappings);
+	}
+
+private:
+#if WITH_DATALINK_CONTEXT
+	FString BuildContextName() const;
+#endif
+
+	void OnExecutionFinished(const FDataLinkExecutor& InExecutor, EDataLinkExecutionResult InResult, FConstStructView InOutputDataView);
+
+	IAvaSceneInterface* GetSceneInterface() const;
+
+	URemoteControlPreset* GetRemoteControlPreset() const;
+
+	struct FResolvedController
+	{
+		const FAvaDataLinkControllerMapping* Mapping;
+		TObjectPtr<URCController> Controller;
+		FProperty* TargetProperty;
+		uint8* TargetMemory;
+	};
+	void ForEachResolvedController(const FDataLinkExecutor& InExecutor, URemoteControlPreset* InPreset, FStructView InTargetDataView, TFunctionRef<void(const FResolvedController&)> InFunction);
+
+	UPROPERTY(EditAnywhere, BlueprintReadWrite, Category="Motion Design Data Link", meta=(AllowPrivateAccess="true"))
+	FDataLinkInstance DataLinkInstance;
+
+	UPROPERTY(EditAnywhere, BlueprintReadWrite, Category="Motion Design Data Link", DisplayName="Output Field to Controller Mappings", meta=(AllowPrivateAccess="true"))
+	TArray<FAvaDataLinkControllerMapping> ControllerMappings;
+
+	TSharedPtr<FDataLinkExecutor> Executor;
+};
diff --git a/Engine/Plugins/Experimental/AvalancheDataLink/Source/AvalancheDataLinkEditor/AvalancheDataLinkEditor.Build.cs b/Engine/Plugins/Experimental/AvalancheDataLink/Source/AvalancheDataLinkEditor/AvalancheDataLinkEditor.Build.cs
new file mode 100644
--- /dev/null
+++ b/Engine/Plugins/Experimental/AvalancheDataLink/Source/AvalancheDataLinkEditor/AvalancheDataLinkEditor.Build.cs
@@ -0,0 +1,24 @@
+// Copyright Epic Games, Inc. All Rights Reserved.
+
+using UnrealBuildTool;
+
+public class AvalancheDataLinkEditor : ModuleRules
+{
+    public AvalancheDataLinkEditor(ReadOnlyTargetRules Target) : base(Target)
+    {
+        PCHUsage = ModuleRules.PCHUsageMode.UseExplicitOrSharedPCHs;
+
+        PrivateDependencyModuleNames.AddRange(
+            new string[]
+            {
+                "AvalancheDataLink",
+                "Core",
+                "CoreUObject",
+                "Engine",
+                "PropertyEditor",
+                "Slate",
+                "SlateCore", 
+            }
+        );
+    }
+}
diff --git a/Engine/Plugins/Experimental/AvalancheDataLink/Source/AvalancheDataLinkEditor/Private/AvaDataLinkEditorModule.cpp b/Engine/Plugins/Experimental/AvalancheDataLink/Source/AvalancheDataLinkEditor/Private/AvaDataLinkEditorModule.cpp
new file mode 100644
--- /dev/null
+++ b/Engine/Plugins/Experimental/AvalancheDataLink/Source/AvalancheDataLinkEditor/Private/AvaDataLinkEditorModule.cpp
@@ -0,0 +1,29 @@
+// Copyright Epic Games, Inc. All Rights Reserved.
+
+#include "AvaDataLinkEditorModule.h"
+#include "AvaDataLinkInstance.h"
+#include "DetailsView/AvaDataLinkInstanceCustomization.h"
+#include "Modules/ModuleManager.h"
+#include "PropertyEditorModule.h"
+
+IMPLEMENT_MODULE(FAvaDataLinkEditorModule, AvalancheDataLinkEditor)
+
+void FAvaDataLinkEditorModule::StartupModule()
+{
+	FPropertyEditorModule& PropertyEditorModule = FModuleManager::LoadModuleChecked<FPropertyEditorModule>("PropertyEditor");
+
+	PropertyEditorModule.RegisterCustomClassLayout(CustomizedClasses.Add_GetRef(UAvaDataLinkInstance::StaticClass()->GetFName())
+		, FOnGetDetailCustomizationInstance::CreateStatic(&FAvaDataLinkInstanceCustomization::MakeInstance));
+}
+
+void FAvaDataLinkEditorModule::ShutdownModule()
+{
+	if (FPropertyEditorModule* PropertyEditorModule = FModuleManager::GetModulePtr<FPropertyEditorModule>("PropertyEditor"))
+	{
+		for (FName CustomizedClass : CustomizedClasses)
+		{
+			PropertyEditorModule->UnregisterCustomClassLayout(CustomizedClass);
+		}
+		CustomizedClasses.Reset();
+	}
+}
diff --git a/Engine/Plugins/Experimental/AvalancheDataLink/Source/AvalancheDataLinkEditor/Private/AvaDataLinkEditorModule.h b/Engine/Plugins/Experimental/AvalancheDataLink/Source/AvalancheDataLinkEditor/Private/AvaDataLinkEditorModule.h
new file mode 100644
--- /dev/null
+++ b/Engine/Plugins/Experimental/AvalancheDataLink/Source/AvalancheDataLinkEditor/Private/AvaDataLinkEditorModule.h
@@ -0,0 +1,17 @@
+// Copyright Epic Games, Inc. All Rights Reserved.
+
+#pragma once
+
+#include "Containers/Array.h"
+#include "Modules/ModuleInterface.h"
+#include "UObject/NameTypes.h"
+
+class FAvaDataLinkEditorModule : public IModuleInterface
+{
+    //~ Begin IModuleInterface
+    virtual void StartupModule() override;
+    virtual void ShutdownModule() override;
+    //~ End IModuleInterface
+
+    TArray<FName> CustomizedClasses;
+};
diff --git a/Engine/Plugins/Experimental/AvalancheDataLink/Source/AvalancheDataLinkEditor/Private/DetailsView/AvaDataLinkControllerMappingsBuilder.cpp b/Engine/Plugins/Experimental/AvalancheDataLink/Source/AvalancheDataLinkEditor/Private/DetailsView/AvaDataLinkControllerMappingsBuilder.cpp
new file mode 100644
--- /dev/null
+++ b/Engine/Plugins/Experimental/AvalancheDataLink/Source/AvalancheDataLinkEditor/Private/DetailsView/AvaDataLinkControllerMappingsBuilder.cpp
@@ -0,0 +1,110 @@
+// Copyright Epic Games, Inc. All Rights Reserved.
+
+#include "AvaDataLinkControllerMappingsBuilder.h"
+#include "AvaDataLinkInstance.h"
+#include "DetailLayoutBuilder.h"
+#include "DetailWidgetRow.h"
+#include "IDetailChildrenBuilder.h"
+
+#define LOCTEXT_NAMESPACE "AvaDataLinkControllerMappingsBuilder"
+
+FAvaDataLinkControllerMappingsBuilder::FAvaDataLinkControllerMappingsBuilder(const TSharedRef<IPropertyHandle>& InControllerMappingsHandle, TWeakPtr<IDetailsView> InDetailsViewWeak)
+	: ControllerMappingsHandle(InControllerMappingsHandle)
+	, ArrayHandle(InControllerMappingsHandle->AsArray().ToSharedRef())
+	, DetailsViewWeak(InDetailsViewWeak)
+{
+}
+
+FName FAvaDataLinkControllerMappingsBuilder::GetName() const
+{
+	return ControllerMappingsHandle->GetProperty()->GetFName();
+}
+
+void FAvaDataLinkControllerMappingsBuilder::SetOnRebuildChildren(FSimpleDelegate InOnRegenerateChildren)
+{
+	OnRegenerateChildren = InOnRegenerateChildren;
+}
+
+void FAvaDataLinkControllerMappingsBuilder::GenerateHeaderRowContent(FDetailWidgetRow& InNodeRow)
+{
+	FUIAction CopyAction;
+	FUIAction PasteAction;
+	ControllerMappingsHandle->CreateDefaultPropertyCopyPasteActions(CopyAction, PasteAction);
+
+	InNodeRow
+		.PropertyHandleList({ ControllerMappingsHandle })
+		.FilterString(ControllerMappingsHandle->GetPropertyDisplayName())
+		.NameContent()
+		[
+			ControllerMappingsHandle->CreatePropertyNameWidget()
+		]
+		.ValueContent()
+		[
+			ControllerMappingsHandle->CreatePropertyValueWidget()
+		]
+		.CopyAction(CopyAction)
+		.PasteAction(PasteAction);
+}
+
+void FAvaDataLinkControllerMappingsBuilder::GenerateChildContent(IDetailChildrenBuilder& InChildrenBuilder)
+{
+	ArrayHandle->GetNumElements(CachedNumChildren);
+
+	TSharedPtr<IDetailsView> DetailsViewShared = DetailsViewWeak.Pin();
+	IDetailsView* DetailsView = DetailsViewShared.Get();
+
+	for (uint32 ChildIndex = 0; ChildIndex < CachedNumChildren; ++ChildIndex)
+	{
+		TSharedRef<IPropertyHandle> ElementHandle = ArrayHandle->GetElement(ChildIndex);
+
+		TSharedPtr<IPropertyHandle> OutputFieldNameHandle = ElementHandle->GetChildHandle(GET_MEMBER_NAME_CHECKED(FAvaDataLinkControllerMapping, OutputFieldName));
+		TSharedPtr<IPropertyHandle> TargetControllerHandle = ElementHandle->GetChildHandle(GET_MEMBER_NAME_CHECKED(FAvaDataLinkControllerMapping, TargetController));
+
+		constexpr bool bDisplayDefaultPropertyButtons = false;
+
+		InChildrenBuilder
+			.AddProperty(ElementHandle)
+			.CustomWidget()
+			.NameContent()
+			[
+				ElementHandle->CreatePropertyNameWidget()
+			]
+			.ValueContent()
+			.MinDesiredWidth(300)
+			[
+				SNew(SHorizontalBox)
+				+ SHorizontalBox::Slot()
+				.FillWidth(1.f)
+				[
+					OutputFieldNameHandle->CreatePropertyValueWidget(bDisplayDefaultPropertyButtons)
+				]
+				+ SHorizontalBox::Slot()
+				.FillWidth(1.f)
+				.Padding(5, 0, 0, 0)
+				[
+					TargetControllerHandle->CreatePropertyValueWidgetWithCustomization(DetailsView)
+				]
+			];
+	}
+}
+
+TSharedPtr<IPropertyHandle> FAvaDataLinkControllerMappingsBuilder::GetPropertyHandle() const
+{
+	return ControllerMappingsHandle;
+}
+
+void FAvaDataLinkControllerMappingsBuilder::Tick(float InDeltaTime)
+{
+	uint32 NumChildren = 0;
+	if (ControllerMappingsHandle->GetNumChildren(NumChildren) != FPropertyAccess::Success)
+	{
+		return;
+	}
+
+	if (NumChildren != CachedNumChildren)
+	{
+		OnRegenerateChildren.ExecuteIfBound();
+	}
+}
+
+#undef LOCTEXT_NAMESPACE
diff --git a/Engine/Plugins/Experimental/AvalancheDataLink/Source/AvalancheDataLinkEditor/Private/DetailsView/AvaDataLinkControllerMappingsBuilder.h b/Engine/Plugins/Experimental/AvalancheDataLink/Source/AvalancheDataLinkEditor/Private/DetailsView/AvaDataLinkControllerMappingsBuilder.h
new file mode 100644
--- /dev/null
+++ b/Engine/Plugins/Experimental/AvalancheDataLink/Source/AvalancheDataLinkEditor/Private/DetailsView/AvaDataLinkControllerMappingsBuilder.h
@@ -0,0 +1,36 @@
+// Copyright Epic Games, Inc. All Rights Reserved.
+
+#pragma once
+
+#include "IDetailCustomNodeBuilder.h"
+
+class IDetailsView;
+class IPropertyHandleArray;
+
+class FAvaDataLinkControllerMappingsBuilder : public IDetailCustomNodeBuilder
+{
+public:
+	explicit FAvaDataLinkControllerMappingsBuilder(const TSharedRef<IPropertyHandle>& InControllerMappingsHandle, TWeakPtr<IDetailsView> InDetailsViewWeak);
+
+	//~ Begin IDetailCustomNodeBuilder
+	virtual FName GetName() const override;
+	virtual void SetOnRebuildChildren(FSimpleDelegate InOnRegenerateChildren) override;
+	virtual void GenerateHeaderRowContent(FDetailWidgetRow& InNodeRow) override;
+	virtual void GenerateChildContent(IDetailChildrenBuilder& InChildrenBuilder) override;
+	virtual TSharedPtr<IPropertyHandle> GetPropertyHandle() const override;
+	virtual bool InitiallyCollapsed() const { return false; }
+	virtual bool RequiresTick() const override { return true; }
+	virtual void Tick(float InDeltaTime) override;
+	//~ End IDetailCustomNodeBuilder
+
+private:
+	TSharedRef<IPropertyHandle> ControllerMappingsHandle;
+
+	TSharedRef<IPropertyHandleArray> ArrayHandle;
+
+	FSimpleDelegate OnRegenerateChildren;
+	
+	TWeakPtr<IDetailsView> DetailsViewWeak;
+
+	uint32 CachedNumChildren = 0;
+};
diff --git a/Engine/Plugins/Experimental/AvalancheDataLink/Source/AvalancheDataLinkEditor/Private/DetailsView/AvaDataLinkInstanceCustomization.cpp b/Engine/Plugins/Experimental/AvalancheDataLink/Source/AvalancheDataLinkEditor/Private/DetailsView/AvaDataLinkInstanceCustomization.cpp
new file mode 100644
--- /dev/null
+++ b/Engine/Plugins/Experimental/AvalancheDataLink/Source/AvalancheDataLinkEditor/Private/DetailsView/AvaDataLinkInstanceCustomization.cpp
@@ -0,0 +1,23 @@
+// Copyright Epic Games, Inc. All Rights Reserved.
+
+#include "AvaDataLinkInstanceCustomization.h"
+
+#include "AvaDataLinkControllerMappingsBuilder.h"
+#include "AvaDataLinkInstance.h"
+#include "DetailCategoryBuilder.h"
+#include "DetailLayoutBuilder.h"
+#include "IDetailChildrenBuilder.h"
+#include "PropertyCustomizationHelpers.h"
+
+void FAvaDataLinkInstanceCustomization::CustomizeDetails(IDetailLayoutBuilder& InDetailBuilder)
+{
+	TSharedRef<IPropertyHandle> DataLinkInstanceHandle = InDetailBuilder.GetProperty(UAvaDataLinkInstance::GetDataLinkInstancePropertyName());
+	DataLinkInstanceHandle->MarkHiddenByCustomization();
+
+	TSharedRef<IPropertyHandle> ControllerMappingsHandle = InDetailBuilder.GetProperty(UAvaDataLinkInstance::GetControllerMappingsPropertyName());
+	ControllerMappingsHandle->MarkHiddenByCustomization();
+
+	IDetailCategoryBuilder& Category = InDetailBuilder.EditCategory(DataLinkInstanceHandle->GetDefaultCategoryName(), DataLinkInstanceHandle->GetDefaultCategoryText());
+	Category.AddProperty(DataLinkInstanceHandle);
+	Category.AddCustomBuilder(MakeShared<FAvaDataLinkControllerMappingsBuilder>(ControllerMappingsHandle, InDetailBuilder.GetDetailsViewSharedPtr()));
+}
diff --git a/Engine/Plugins/Experimental/AvalancheDataLink/Source/AvalancheDataLinkEditor/Private/DetailsView/AvaDataLinkInstanceCustomization.h b/Engine/Plugins/Experimental/AvalancheDataLink/Source/AvalancheDataLinkEditor/Private/DetailsView/AvaDataLinkInstanceCustomization.h
new file mode 100644
--- /dev/null
+++ b/Engine/Plugins/Experimental/AvalancheDataLink/Source/AvalancheDataLinkEditor/Private/DetailsView/AvaDataLinkInstanceCustomization.h
@@ -0,0 +1,19 @@
+// Copyright Epic Games, Inc. All Rights Reserved.
+
+#pragma once
+
+#include "IDetailCustomization.h"
+
+class FAvaDataLinkInstanceCustomization : public IDetailCustomization
+{
+public:
+	static TSharedRef<IDetailCustomization> MakeInstance()
+	{
+		return MakeShared<FAvaDataLinkInstanceCustomization>();
+	}
+
+protected:
+	//~ Begin IDetailCustomization
+	virtual void CustomizeDetails(IDetailLayoutBuilder& InDetailBuilder) override;
+	//~ End IDetailCustomization
+};
diff --git a/Engine/Plugins/Experimental/AvalancheSceneState/AvalancheSceneState.uplugin b/Engine/Plugins/Experimental/AvalancheSceneState/AvalancheSceneState.uplugin
new file mode 100644
--- /dev/null
+++ b/Engine/Plugins/Experimental/AvalancheSceneState/AvalancheSceneState.uplugin
@@ -0,0 +1,45 @@
+{
+	"FileVersion": 3,
+	"Version": 1,
+	"VersionName": "1.0",
+	"FriendlyName": "Motion Design Scene State Integration",
+	"Description": "",
+	"Category": "Other",
+	"CreatedBy": "",
+	"CreatedByURL": "",
+	"DocsURL": "",
+	"MarketplaceURL": "",
+	"CanContainContent": true,
+	"IsBetaVersion": false,
+	"IsExperimentalVersion": true,
+	"Installed": false,
+	"Modules": [
+		{
+			"Name": "AvalancheSceneState",
+			"Type": "Runtime",
+			"LoadingPhase": "Default",
+			"TargetDenyList": [ "Server" ]
+		},
+		{
+			"Name": "AvalancheSceneStateBlueprint",
+			"Type": "UncookedOnly",
+			"LoadingPhase": "Default",
+			"TargetDenyList": [ "Server" ]
+		},
+		{
+			"Name": "AvalancheSceneStateEditor",
+			"Type": "Editor",
+			"LoadingPhase": "Default"
+		}
+	],
+	"Plugins": [
+		{
+			"Name": "Avalanche",
+			"Enabled": true
+		},
+		{
+			"Name": "SceneState",
+			"Enabled": true
+		}
+	]
+}
\ No newline at end of file
diff --git a/Engine/Plugins/Experimental/AvalancheSceneState/Source/AvalancheSceneState/AvalancheSceneState.Build.cs b/Engine/Plugins/Experimental/AvalancheSceneState/Source/AvalancheSceneState/AvalancheSceneState.Build.cs
new file mode 100644
--- /dev/null
+++ b/Engine/Plugins/Experimental/AvalancheSceneState/Source/AvalancheSceneState/AvalancheSceneState.Build.cs
@@ -0,0 +1,46 @@
+// Copyright Epic Games, Inc. All Rights Reserved.
+
+using UnrealBuildTool;
+
+public class AvalancheSceneState : ModuleRules
+{
+	public AvalancheSceneState(ReadOnlyTargetRules Target) : base(Target)
+	{
+		PCHUsage = ModuleRules.PCHUsageMode.UseExplicitOrSharedPCHs;
+
+		PublicDependencyModuleNames.AddRange(
+			new string[]
+			{
+				"AvalancheRemoteControl",
+				"AvalancheSequence",
+				"AvalancheTag",
+				"Core",
+				"CoreUObject",
+				"SceneState",
+				"SceneStateBinding",
+				"SceneStateGameplay",
+			}
+		);
+
+		PrivateDependencyModuleNames.AddRange(
+			new string[]
+			{
+				"Avalanche",
+				"Engine",
+				"RemoteControl",
+				"RemoteControlLogic",
+				"SceneStateEvent",
+			}
+		);
+
+		if (Target.Type == TargetRules.TargetType.Editor)
+		{
+			PrivateDependencyModuleNames.AddRange(
+				new string[]
+				{
+					"SceneStateBlueprint",
+					"SceneStateBlueprintEditor",
+				});
+		}
+	}
+}
\ No newline at end of file
diff --git a/Engine/Plugins/Experimental/AvalancheSceneState/Source/AvalancheSceneState/Private/AvaSceneStateActor.cpp b/Engine/Plugins/Experimental/AvalancheSceneState/Source/AvalancheSceneState/Private/AvaSceneStateActor.cpp
new file mode 100644
--- /dev/null
+++ b/Engine/Plugins/Experimental/AvalancheSceneState/Source/AvalancheSceneState/Private/AvaSceneStateActor.cpp
@@ -0,0 +1,112 @@
+// Copyright Epic Games, Inc. All Rights Reserved.
+
+#include "AvaSceneStateActor.h"
+#include "AvaSceneStateComponent.h"
+#include "Engine/World.h"
+#include "SceneStateGeneratedClass.h"
+
+#if WITH_EDITOR
+#include "SceneStateBlueprint.h"
+#include "SceneStateBlueprintFactory.h"
+#endif
+
+AAvaSceneStateActor::AAvaSceneStateActor(const FObjectInitializer& InObjectInitializer)
+	: Super(InObjectInitializer.SetDefaultSubobjectClass<UAvaSceneStateComponent>(ASceneStateActor::SceneStateComponentName))
+{
+#if WITH_EDITOR
+	FWorldDelegates::OnWorldCleanup.AddUObject(this, &AAvaSceneStateActor::OnWorldCleanup);
+#endif
+}
+
+#if WITH_EDITOR
+FString AAvaSceneStateActor::GetDefaultActorLabel() const
+{
+	return TEXT("Motion Design Scene State");
+}
+
+void AAvaSceneStateActor::PostLoad()
+{
+	Super::PostLoad();
+	bListedInSceneOutliner = true;
+	SetSceneStateBlueprint(Cast<USceneStateBlueprint>(SceneStateBlueprint));
+}
+
+void AAvaSceneStateActor::PostDuplicate(bool bInDuplicateForPIE)
+{
+	Super::PostDuplicate(bInDuplicateForPIE);
+
+	if (bInDuplicateForPIE)
+	{
+		SetSceneStateBlueprint(nullptr);
+	}
+	else
+	{
+		SetSceneStateBlueprint(Cast<USceneStateBlueprint>(SceneStateBlueprint));
+		UpdateSceneStateClass();
+	}
+}
+
+void AAvaSceneStateActor::BeginDestroy()
+{
+	Super::BeginDestroy();
+	FWorldDelegates::OnWorldCleanup.RemoveAll(this);
+}
+
+void AAvaSceneStateActor::UpdateSceneStateClass()
+{
+	if (SceneStateBlueprint)
+	{
+		SetSceneStateClass(Cast<USceneStateGeneratedClass>(SceneStateBlueprint->GeneratedClass));
+	}
+	else
+	{
+		SetSceneStateClass(nullptr);
+	}
+}
+
+void AAvaSceneStateActor::SetSceneStateBlueprint(USceneStateBlueprint* InSceneStateBlueprint)
+{
+	if (SceneStateBlueprint)
+	{
+		SceneStateBlueprint->OnCompiled().RemoveAll(this);
+	}
+
+	SceneStateBlueprint = InSceneStateBlueprint;
+
+	if (SceneStateBlueprint)
+	{
+		SceneStateBlueprint->OnCompiled().AddUObject(this, &AAvaSceneStateActor::OnSceneStateRecompiled);
+	}
+}
+
+void AAvaSceneStateActor::OnSceneStateRecompiled(UBlueprint* InCompiledBlueprint)
+{
+	ensure(InCompiledBlueprint == SceneStateBlueprint);
+	UpdateSceneStateClass();
+}
+
+void AAvaSceneStateActor::OnWorldCleanup(UWorld* InWorld, bool bInSessionEnded, bool bInCleanupResources)
+{
+	if (!bInCleanupResources || !SceneStateBlueprint)
+	{
+		return;
+	}
+
+	// Ignore cleanups from other worlds
+	if (GetTypedOuter<UWorld>() != InWorld)
+	{
+		return;
+	}
+
+	const FName ObjectName = MakeUniqueObjectName(GetTransientPackage()
+		, SceneStateBlueprint->GetClass()
+		, FName(*FString::Printf(TEXT("%s_Trashed"), *SceneStateBlueprint->GetName())));
+
+	SceneStateBlueprint->Rename(*ObjectName.ToString()
+		, GetTransientPackage()
+		, REN_DontCreateRedirectors | REN_NonTransactional | REN_DoNotDirty);
+
+	SetSceneStateBlueprint(nullptr);
+	SetSceneStateClass(nullptr);
+}
+#endif
diff --git a/Engine/Plugins/Experimental/AvalancheSceneState/Source/AvalancheSceneState/Private/AvaSceneStateComponent.cpp b/Engine/Plugins/Experimental/AvalancheSceneState/Source/AvalancheSceneState/Private/AvaSceneStateComponent.cpp
new file mode 100644
--- /dev/null
+++ b/Engine/Plugins/Experimental/AvalancheSceneState/Source/AvalancheSceneState/Private/AvaSceneStateComponent.cpp
@@ -0,0 +1,9 @@
+// Copyright Epic Games, Inc. All Rights Reserved.
+
+#include "AvaSceneStateComponent.h"
+#include "AvaSceneStatePlayer.h"
+
+UAvaSceneStateComponent::UAvaSceneStateComponent(const FObjectInitializer& InObjectInitializer)
+	: Super(InObjectInitializer.SetDefaultSubobjectClass<UAvaSceneStatePlayer>(UAvaSceneStateComponent::SceneStatePlayerName))
+{
+}
diff --git a/Engine/Plugins/Experimental/AvalancheSceneState/Source/AvalancheSceneState/Private/AvaSceneStateLog.h b/Engine/Plugins/Experimental/AvalancheSceneState/Source/AvalancheSceneState/Private/AvaSceneStateLog.h
new file mode 100644
--- /dev/null
+++ b/Engine/Plugins/Experimental/AvalancheSceneState/Source/AvalancheSceneState/Private/AvaSceneStateLog.h
@@ -0,0 +1,7 @@
+// Copyright Epic Games, Inc. All Rights Reserved.
+
+#pragma once
+
+#include "Logging/LogMacros.h"
+
+DECLARE_LOG_CATEGORY_EXTERN(LogAvaSceneState, Log, All);
diff --git a/Engine/Plugins/Experimental/AvalancheSceneState/Source/AvalancheSceneState/Private/AvaSceneStateModule.cpp b/Engine/Plugins/Experimental/AvalancheSceneState/Source/AvalancheSceneState/Private/AvaSceneStateModule.cpp
new file mode 100644
--- /dev/null
+++ b/Engine/Plugins/Experimental/AvalancheSceneState/Source/AvalancheSceneState/Private/AvaSceneStateModule.cpp
@@ -0,0 +1,8 @@
+// Copyright Epic Games, Inc. All Rights Reserved.
+
+#include "AvaSceneStateLog.h"
+#include "Modules/ModuleManager.h"
+
+DEFINE_LOG_CATEGORY(LogAvaSceneState);
+
+IMPLEMENT_MODULE(FDefaultModuleImpl, AvalancheSceneState)
diff --git a/Engine/Plugins/Experimental/AvalancheSceneState/Source/AvalancheSceneState/Private/AvaSceneStatePlayer.cpp b/Engine/Plugins/Experimental/AvalancheSceneState/Source/AvalancheSceneState/Private/AvaSceneStatePlayer.cpp
new file mode 100644
--- /dev/null
+++ b/Engine/Plugins/Experimental/AvalancheSceneState/Source/AvalancheSceneState/Private/AvaSceneStatePlayer.cpp
@@ -0,0 +1,10 @@
+// Copyright Epic Games, Inc. All Rights Reserved.
+
+#include "AvaSceneStatePlayer.h"
+
+UAvaSceneStatePlayer::UAvaSceneStatePlayer()
+{
+#if WITH_EDITOR
+	bEditableSceneStateClass = false;
+#endif
+}
diff --git a/Engine/Plugins/Experimental/AvalancheSceneState/Source/AvalancheSceneState/Private/AvaSceneStateUtils.cpp b/Engine/Plugins/Experimental/AvalancheSceneState/Source/AvalancheSceneState/Private/AvaSceneStateUtils.cpp
new file mode 100644
--- /dev/null
+++ b/Engine/Plugins/Experimental/AvalancheSceneState/Source/AvalancheSceneState/Private/AvaSceneStateUtils.cpp
@@ -0,0 +1,37 @@
+// Copyright Epic Games, Inc. All Rights Reserved.
+
+#include "AvaSceneStateUtils.h"
+#include "AvaSceneStateLog.h"
+#include "AvaSceneSubsystem.h"
+#include "Engine/Level.h"
+#include "SceneStateExecutionContext.h"
+
+namespace UE::AvaSceneState
+{
+
+IAvaSceneInterface* FindSceneInterface(const FSceneStateExecutionContext& InContext)
+{
+	UObject* ContextObject = InContext.GetContextObject();
+	if (!ContextObject)
+	{
+		UE_LOG(LogAvaSceneState, Warning, TEXT("[%s] Scene State could not find a valid Context Object"), *InContext.GetExecutionContextName());
+		return nullptr;
+	}
+
+	ULevel* ContextLevel = Cast<ULevel>(ContextObject);
+	if (!ContextLevel)
+	{
+		ContextLevel = ContextObject->GetTypedOuter<ULevel>();
+	}
+
+	IAvaSceneInterface* SceneInterface = UAvaSceneSubsystem::FindSceneInterface(ContextLevel);
+	if (!SceneInterface)
+	{
+		UE_LOG(LogAvaSceneState, Warning, TEXT("[%s] Failed to find Motion Design Scene Interface!"), *InContext.GetExecutionContextName());
+		return nullptr;
+	}
+
+	return SceneInterface;
+}
+
+} // UE::AvaSceneState
diff --git a/Engine/Plugins/Experimental/AvalancheSceneState/Source/AvalancheSceneState/Private/AvaSceneStateUtils.h b/Engine/Plugins/Experimental/AvalancheSceneState/Source/AvalancheSceneState/Private/AvaSceneStateUtils.h
new file mode 100644
--- /dev/null
+++ b/Engine/Plugins/Experimental/AvalancheSceneState/Source/AvalancheSceneState/Private/AvaSceneStateUtils.h
@@ -0,0 +1,11 @@
+// Copyright Epic Games, Inc. All Rights Reserved.
+
+#pragma once
+
+class IAvaSceneInterface;
+struct FSceneStateExecutionContext;
+
+namespace UE::AvaSceneState
+{
+	IAvaSceneInterface* FindSceneInterface(const FSceneStateExecutionContext& InContext);
+}
diff --git a/Engine/Plugins/Experimental/AvalancheSceneState/Source/AvalancheSceneState/Private/RemoteControl/AvaSceneStateRCEventBehaviorNode.cpp b/Engine/Plugins/Experimental/AvalancheSceneState/Source/AvalancheSceneState/Private/RemoteControl/AvaSceneStateRCEventBehaviorNode.cpp
new file mode 100644
--- /dev/null
+++ b/Engine/Plugins/Experimental/AvalancheSceneState/Source/AvalancheSceneState/Private/RemoteControl/AvaSceneStateRCEventBehaviorNode.cpp
@@ -0,0 +1,36 @@
+// Copyright Epic Games, Inc. All Rights Reserved.
+
+#include "RemoteControl/AvaSceneStateRCEventBehaviorNode.h"
+#include "RemoteControl/AvaSceneStateRCEventBehavior.h"
+#include "SceneStateEventUtils.h"
+#include "StructUtils/UserDefinedStruct.h"
+
+#define LOCTEXT_NAMESPACE "AvaSceneStateRCEventBehaviorNode"
+
+UAvaSceneStateRCEventBehaviorNode::UAvaSceneStateRCEventBehaviorNode()
+{
+	DisplayName = LOCTEXT("DisplayName", "Broadcast Event");
+	BehaviorDescription = LOCTEXT("Description", "Broadcasts a defined event after applying the actions");
+}
+
+bool UAvaSceneStateRCEventBehaviorNode::Execute(URCBehaviour* InBehavior) const
+{
+	return true;
+}
+
+bool UAvaSceneStateRCEventBehaviorNode::IsSupported(URCBehaviour* InBehavior) const
+{
+	return InBehavior && InBehavior->IsA<UAvaSceneStateRCEventBehavior>();
+}
+
+void UAvaSceneStateRCEventBehaviorNode::OnPassed(URCBehaviour* InBehavior) const
+{
+	UE::SceneState::BroadcastEvent(Event);
+}
+
+UClass* UAvaSceneStateRCEventBehaviorNode::GetBehaviourClass() const
+{
+	return UAvaSceneStateRCEventBehavior::StaticClass();
+}
+
+#undef LOCTEXT_NAMESPACE
diff --git a/Engine/Plugins/Experimental/AvalancheSceneState/Source/AvalancheSceneState/Private/RemoteControl/AvaSceneStateRCTask.cpp b/Engine/Plugins/Experimental/AvalancheSceneState/Source/AvalancheSceneState/Private/RemoteControl/AvaSceneStateRCTask.cpp
new file mode 100644
--- /dev/null
+++ b/Engine/Plugins/Experimental/AvalancheSceneState/Source/AvalancheSceneState/Private/RemoteControl/AvaSceneStateRCTask.cpp
@@ -0,0 +1,152 @@
+// Copyright Epic Games, Inc. All Rights Reserved.
+
+#include "RemoteControl/AvaSceneStateRCTask.h"
+#include "AvaSceneStateLog.h"
+#include "AvaSceneStateUtils.h"
+#include "Controller/RCController.h"
+#include "IAvaSceneInterface.h"
+#include "RCVirtualPropertyContainer.h"
+#include "RemoteControlPreset.h"
+#include "SceneStateBindingDesc.h"
+#include "SceneStateBindingUtils.h"
+#include "SceneStateExecutionContext.h"
+#include "UObject/PropertyAccessUtil.h"
+
+FAvaSceneStateRCTask::FAvaSceneStateRCTask()
+{
+	SetFlags(ESceneStateTaskFlags::HasBindingExtension);
+}
+
+#if WITH_EDITOR
+const UScriptStruct* FAvaSceneStateRCTask::OnGetTaskInstanceType() const
+{
+	return FInstanceDataType::StaticStruct();
+}
+
+void FAvaSceneStateRCTask::OnBuildTaskInstance(UObject* InOuter, FStructView InTaskInstance) const
+{
+	FInstanceDataType& Instance = InTaskInstance.Get<FInstanceDataType>();
+
+	const FGuid OldControllerValuesId = Instance.ControllerValuesId;
+	Instance.ControllerValuesId = FGuid::NewGuid();
+
+	UE::SceneState::HandleStructIdChanged(*InOuter, OldControllerValuesId, Instance.ControllerValuesId);
+}
+#endif
+
+const FSceneStateTaskBindingExtension* FAvaSceneStateRCTask::OnGetBindingExtension() const
+{
+	return &Binding;
+}
+
+void FAvaSceneStateRCTask::OnStart(const FSceneStateExecutionContext& InContext, FStructView InTaskInstance) const
+{
+	ON_SCOPE_EXIT
+	{
+		Finish(InContext, InTaskInstance);
+	};
+
+	FInstanceDataType& Instance = InTaskInstance.Get<FInstanceDataType>();
+
+	const UPropertyBag* SourcePropertyBag = Instance.ControllerValues.GetPropertyBagStruct();
+	const uint8* SourceValuesMemory = Instance.ControllerValues.GetValue().GetMemory();
+	if (!SourcePropertyBag)
+	{
+		UE_LOG(LogAvaSceneState, Warning, TEXT("[%s] No valid source data to copy from!"), *InContext.GetExecutionContextName());
+		return;
+	}
+
+	URemoteControlPreset* Preset;
+	FStructView TargetControllerDataView;
+	if (!GetControllerDataView(InContext, TargetControllerDataView, Preset))
+	{
+		return;
+	}
+
+	// At most, each entry in the source property bag will copy to a controller
+	TArray<URCController*> ModifiedControllers;
+	ModifiedControllers.Reserve(Instance.ControllerValues.GetNumPropertiesInBag());
+
+	for (const FPropertyBagPropertyDesc& SourcePropertyDesc : SourcePropertyBag->GetPropertyDescs())
+	{
+		if (!SourcePropertyDesc.CachedProperty)
+		{
+			UE_LOG(LogAvaSceneState, Warning, TEXT("[%s] Property '%s' was not invalid in the controller values!")
+				, *InContext.GetExecutionContextName()
+				, *SourcePropertyDesc.Name.ToString());
+			continue;
+		}
+
+		const FAvaSceneStateRCControllerMapping* ControllerMapping = Instance.ControllerMappings.FindByKey(SourcePropertyDesc.ID);
+		if (!ControllerMapping)
+		{
+			UE_LOG(LogAvaSceneState, Warning, TEXT("[%s] Property '%s' was not found in the controller mappings!")
+				, *InContext.GetExecutionContextName()
+				, *SourcePropertyDesc.Name.ToString());
+			continue;
+		}
+
+		URCController* const Controller = Cast<URCController>(ControllerMapping->TargetController.FindController(Preset));
+		if (!Controller)
+		{
+			UE_LOG(LogAvaSceneState, Warning, TEXT("[%s] Controller '%s' was not found in preset '%s'!")
+				, *InContext.GetExecutionContextName()
+				, *SourcePropertyDesc.Name.ToString()
+				, *Preset->GetName());
+			continue;
+		}
+
+		FProperty* const TargetProperty = Controller->GetProperty();
+		if (!ensure(TargetProperty) || !TargetProperty->SameType(SourcePropertyDesc.CachedProperty))
+		{
+			UE_LOG(LogAvaSceneState, Warning, TEXT("[%s] Controller '%s' type mismatch in preset '%s'")
+				, *InContext.GetExecutionContextName()
+				, *Controller->DisplayName.ToString()
+				, *Preset->GetName());
+			continue;
+		}
+
+		const void* SourceData = SourcePropertyDesc.CachedProperty->ContainerPtrToValuePtr<void>(SourceValuesMemory);
+		void* TargetData = TargetProperty->ContainerPtrToValuePtr<void>(TargetControllerDataView.GetMemory());
+
+		TargetProperty->CopyCompleteValue(TargetData, SourceData);
+		ModifiedControllers.AddUnique(Controller);
+	}
+
+	TSet<FGuid> ModifiedControllerIds;
+	ModifiedControllerIds.Reserve(ModifiedControllers.Num());
+
+	for (URCController* Controller : ModifiedControllers)
+	{
+		Controller->OnModifyPropertyValue();
+		ModifiedControllerIds.Add(Controller->Id);
+	}
+
+	Preset->OnControllerModified().Broadcast(Preset, ModifiedControllerIds);
+}
+
+bool FAvaSceneStateRCTask::GetControllerDataView(const FSceneStateExecutionContext& InContext, FStructView& OutControllerView, URemoteControlPreset*& OutPreset) const
+{
+	const IAvaSceneInterface* const SceneInterface = UE::AvaSceneState::FindSceneInterface(InContext);
+	if (!SceneInterface)
+	{
+		return false;
+	}
+
+	OutPreset = SceneInterface->GetRemoteControlPreset();
+	if (!OutPreset)
+	{
+		UE_LOG(LogAvaSceneState, Warning, TEXT("[%s] Failed to find Remote Control Preset in Scene!"), *InContext.GetExecutionContextName());
+		return false;
+	}
+
+	URCVirtualPropertyContainerBase* const ControllerContainer = OutPreset->GetControllerContainer();
+	if (!ControllerContainer)
+	{
+		UE_LOG(LogAvaSceneState, Warning, TEXT("[%s] Failed to find Remote Control Preset in Scene!"), *InContext.GetExecutionContextName());
+		return false;
+	}
+
+	OutControllerView = ControllerContainer->GetPropertyBagMutableValue();
+	return OutControllerView.IsValid();
+}
diff --git a/Engine/Plugins/Experimental/AvalancheSceneState/Source/AvalancheSceneState/Private/RemoteControl/AvaSceneStateRCTaskBinding.cpp b/Engine/Plugins/Experimental/AvalancheSceneState/Source/AvalancheSceneState/Private/RemoteControl/AvaSceneStateRCTaskBinding.cpp
new file mode 100644
--- /dev/null
+++ b/Engine/Plugins/Experimental/AvalancheSceneState/Source/AvalancheSceneState/Private/RemoteControl/AvaSceneStateRCTaskBinding.cpp
@@ -0,0 +1,64 @@
+// Copyright Epic Games, Inc. All Rights Reserved.
+
+#include "RemoteControl/AvaSceneStateRCTaskBinding.h"
+#include "RemoteControl/AvaSceneStateRCTask.h"
+
+#if WITH_EDITOR
+void FAvaSceneStateRCTaskBinding::VisitBindingDescs(FConstStructView InTaskInstance, TFunctionRef<void(const UE::SceneState::FTaskBindingDesc&)> InFunctor) const
+{
+	const FInstanceDataType& Instance = InTaskInstance.Get<const FInstanceDataType>();
+
+	if (Instance.ControllerValues.IsValid())
+	{
+		UE::SceneState::FTaskBindingDesc BindingDesc;
+
+		BindingDesc.Id = Instance.ControllerValuesId;
+		BindingDesc.Name = TEXT("ControllerValues");
+		BindingDesc.Struct = Instance.ControllerValues.GetPropertyBagStruct();
+		BindingDesc.DataIndex = CONTROLLER_VALUES_DATA_INDEX;
+
+		InFunctor(BindingDesc);
+	}
+}
+
+void FAvaSceneStateRCTaskBinding::SetBindingBatch(uint16 InDataIndex, uint16 InBatchIndex)
+{
+	if (InDataIndex == CONTROLLER_VALUES_DATA_INDEX)
+	{
+		ControllerValuesBatchIndex = InBatchIndex;
+	}
+}
+
+bool FAvaSceneStateRCTaskBinding::FindDataById(FStructView InTaskInstance, const FGuid& InStructId, FStructView& OutDataView, uint16& OutDataIndex) const
+{
+	FInstanceDataType& Instance = InTaskInstance.Get<FInstanceDataType>();
+	if (Instance.ControllerValuesId == InStructId && Instance.ControllerValues.IsValid())
+	{
+		OutDataView = Instance.ControllerValues.GetMutableValue();
+		OutDataIndex = CONTROLLER_VALUES_DATA_INDEX;
+		return true;
+	}
+	return false;
+}
+#endif
+
+bool FAvaSceneStateRCTaskBinding::FindDataByIndex(FStructView InTaskInstance, uint16 InDataIndex, FStructView& OutDataView) const
+{
+	FInstanceDataType& Instance = InTaskInstance.Get<FInstanceDataType>();
+	if (InDataIndex == CONTROLLER_VALUES_DATA_INDEX && Instance.ControllerValues.IsValid())
+	{
+		OutDataView = Instance.ControllerValues.GetMutableValue();
+		return true;
+	}
+	return false;
+}
+
+void FAvaSceneStateRCTaskBinding::VisitBindingBatches(FStructView InTaskInstance, TFunctionRef<void(uint16, FStructView)> InFunctor) const
+{
+	FInstanceDataType& Instance = InTaskInstance.Get<FInstanceDataType>();
+
+	if (Instance.ControllerValues.IsValid())
+	{
+		InFunctor(ControllerValuesBatchIndex, Instance.ControllerValues.GetMutableValue());
+	}
+}
diff --git a/Engine/Plugins/Experimental/AvalancheSceneState/Source/AvalancheSceneState/Private/Sequence/AvaSceneStatePlaySequenceTask.cpp b/Engine/Plugins/Experimental/AvalancheSceneState/Source/AvalancheSceneState/Private/Sequence/AvaSceneStatePlaySequenceTask.cpp
new file mode 100644
--- /dev/null
+++ b/Engine/Plugins/Experimental/AvalancheSceneState/Source/AvalancheSceneState/Private/Sequence/AvaSceneStatePlaySequenceTask.cpp
@@ -0,0 +1,81 @@
+// Copyright Epic Games, Inc. All Rights Reserved.
+
+#include "Sequence/AvaSceneStatePlaySequenceTask.h"
+#include "AvaSceneStateLog.h"
+#include "AvaSceneStateUtils.h"
+#include "AvaSequence.h"
+#include "AvaSequencePlayer.h"
+#include "AvaSequenceSubsystem.h"
+#include "IAvaSceneInterface.h"
+#include "SceneStateExecutionContext.h"
+#include "Tasks/SceneStateTaskExecutionContext.h"
+
+#if WITH_EDITOR
+const UScriptStruct* FAvaSceneStatePlaySequenceTask::OnGetTaskInstanceType() const
+{
+	return FInstanceDataType::StaticStruct();
+}
+#endif
+
+void FAvaSceneStatePlaySequenceTask::OnStart(const FSceneStateExecutionContext& InContext, FStructView InTaskInstance) const
+{
+	const IAvaSceneInterface* const SceneInterface = UE::AvaSceneState::FindSceneInterface(InContext);
+	if (!SceneInterface)
+	{
+		Finish(InContext, InTaskInstance);
+		return;
+	}
+
+	IAvaSequencePlaybackObject* const PlaybackObject = SceneInterface->GetPlaybackObject();
+	if (!PlaybackObject)
+	{
+		UE_LOG(LogAvaSceneState, Warning, TEXT("[%s] Failed to find Playback Object in Scene Interface!"), *InContext.GetExecutionContextName());
+		Finish(InContext, InTaskInstance);
+		return;
+	}
+
+	FInstanceDataType& Instance = InTaskInstance.Get<FInstanceDataType>();
+
+	Instance.OnSequenceFinishedHandle = UAvaSequencePlayer::OnSequenceFinished()
+		.AddStatic(&FAvaSceneStatePlaySequenceTask::OnSequenceStopped, UE::SceneState::FTaskExecutionContext(*this, InContext));
+
+	switch (Instance.SequenceQueryType)
+	{
+	case EAvaTransitionSequenceQueryType::Name:
+		Instance.SequencePlayers = PlaybackObject->PlaySequencesByLabel(Instance.SequenceName, Instance.PlaySettings);
+		break;
+
+	case EAvaTransitionSequenceQueryType::Tag:
+		Instance.SequencePlayers = PlaybackObject->PlaySequencesByTag(Instance.SequenceTag, /*bPerformExactMatch*/true, Instance.PlaySettings);
+		break;
+	}
+
+	if (Instance.SequencePlayers.IsEmpty() || Instance.WaitType == EAvaTransitionSequenceWaitType::NoWait)
+	{
+		Finish(InContext, InTaskInstance);
+	}
+}
+
+void FAvaSceneStatePlaySequenceTask::OnStop(const FSceneStateExecutionContext& InContext, FStructView InTaskInstance, ESceneStateTaskStopReason InStopReason) const
+{
+	FInstanceDataType& Instance = InTaskInstance.Get<FInstanceDataType>();
+
+	UAvaSequencePlayer::OnSequenceFinished().Remove(Instance.OnSequenceFinishedHandle);
+	Instance.OnSequenceFinishedHandle.Reset();
+}
+
+void FAvaSceneStatePlaySequenceTask::OnSequenceStopped(UAvaSequencePlayer* InPlayer, UAvaSequence* InSequence, UE::SceneState::FTaskExecutionContext InTaskContext)
+{
+	FInstanceDataType* Instance = InTaskContext.GetTaskInstance().GetPtr<FInstanceDataType>();
+	if (!Instance)
+	{
+		InTaskContext.FinishTask();
+		return;
+	}
+
+	Instance->SequencePlayers.Remove(InPlayer);
+	if (Instance->SequencePlayers.IsEmpty())
+	{
+		InTaskContext.FinishTask();
+	}
+}
diff --git a/Engine/Plugins/Experimental/AvalancheSceneState/Source/AvalancheSceneState/Public/AvaSceneStateActor.h b/Engine/Plugins/Experimental/AvalancheSceneState/Source/AvalancheSceneState/Public/AvaSceneStateActor.h
new file mode 100644
--- /dev/null
+++ b/Engine/Plugins/Experimental/AvalancheSceneState/Source/AvalancheSceneState/Public/AvaSceneStateActor.h
@@ -0,0 +1,53 @@
+// Copyright Epic Games, Inc. All Rights Reserved.
+
+#pragma once
+
+#include "GameFramework/Actor.h"
+#include "SceneStateActor.h"
+#include "AvaSceneStateActor.generated.h"
+
+#if WITH_EDITOR
+class UBlueprint;
+class USceneStateBlueprint;
+#endif
+
+/** Scene State Actor that has additional editor-only data holding the Scene State Blueprint */
+UCLASS(MinimalAPI, NotBlueprintable, NotPlaceable, DisplayName="Motion Design Scene State Actor", HideCategories=(Actor, Collision, Cooking, DataLayers, HLOD, Input, LevelInstance, Networking, Physics, Rendering, Replication, WorldPartition))
+class AAvaSceneStateActor : public ASceneStateActor
+{
+	GENERATED_BODY()
+
+public:
+	AAvaSceneStateActor(const FObjectInitializer& InObjectInitializer);
+
+protected:
+#if WITH_EDITOR
+	//~ Begin AActor
+	virtual FString GetDefaultActorLabel() const override;
+	virtual bool IsUserManaged() const override { return false; }
+	virtual bool SupportsExternalPackaging() const override { return false; }
+	//~ End AActor
+
+	//~ Begin UObject
+	virtual void PostLoad() override;
+	virtual void PostDuplicate(bool bInDuplicateForPIE) override;
+	virtual void BeginDestroy() override;
+	//~ End UObject
+
+	AVALANCHESCENESTATE_API void UpdateSceneStateClass();
+
+	AVALANCHESCENESTATE_API void SetSceneStateBlueprint(USceneStateBlueprint* InSceneStateBlueprint);
+
+	void OnSceneStateRecompiled(UBlueprint* InCompiledBlueprint);
+
+	void OnWorldCleanup(UWorld* InWorld, bool bInSessionEnded, bool bInCleanupResources);
+#endif
+
+private:
+#if WITH_EDITORONLY_DATA
+	UPROPERTY()
+	TObjectPtr<UBlueprint> SceneStateBlueprint;
+#endif
+
+	friend class FAvaSceneStateExtension;
+};
diff --git a/Engine/Plugins/Experimental/AvalancheSceneState/Source/AvalancheSceneState/Public/AvaSceneStateComponent.h b/Engine/Plugins/Experimental/AvalancheSceneState/Source/AvalancheSceneState/Public/AvaSceneStateComponent.h
new file mode 100644
--- /dev/null
+++ b/Engine/Plugins/Experimental/AvalancheSceneState/Source/AvalancheSceneState/Public/AvaSceneStateComponent.h
@@ -0,0 +1,14 @@
+// Copyright Epic Games, Inc. All Rights Reserved.
+
+#pragma once
+
+#include "SceneStateComponent.h"
+#include "AvaSceneStateComponent.generated.h"
+
+UCLASS(HideCategories=(Activation, AssetUserData, Cooking, Input, Navigation, Tags))
+class UAvaSceneStateComponent : public USceneStateComponent
+{
+	GENERATED_BODY()
+
+	UAvaSceneStateComponent(const FObjectInitializer& InObjectInitializer);
+};
diff --git a/Engine/Plugins/Experimental/AvalancheSceneState/Source/AvalancheSceneState/Public/AvaSceneStatePlayer.h b/Engine/Plugins/Experimental/AvalancheSceneState/Source/AvalancheSceneState/Public/AvaSceneStatePlayer.h
new file mode 100644
--- /dev/null
+++ b/Engine/Plugins/Experimental/AvalancheSceneState/Source/AvalancheSceneState/Public/AvaSceneStatePlayer.h
@@ -0,0 +1,14 @@
+// Copyright Epic Games, Inc. All Rights Reserved.
+
+#pragma once
+
+#include "SceneStateComponentPlayer.h"
+#include "AvaSceneStatePlayer.generated.h"
+
+UCLASS()
+class UAvaSceneStatePlayer : public USceneStateComponentPlayer
+{
+	GENERATED_BODY()
+
+	UAvaSceneStatePlayer();
+};
diff --git a/Engine/Plugins/Experimental/AvalancheSceneState/Source/AvalancheSceneState/Public/RemoteControl/AvaSceneStateRCEventBehavior.h b/Engine/Plugins/Experimental/AvalancheSceneState/Source/AvalancheSceneState/Public/RemoteControl/AvaSceneStateRCEventBehavior.h
new file mode 100644
--- /dev/null
+++ b/Engine/Plugins/Experimental/AvalancheSceneState/Source/AvalancheSceneState/Public/RemoteControl/AvaSceneStateRCEventBehavior.h
@@ -0,0 +1,12 @@
+// Copyright Epic Games, Inc. All Rights Reserved.
+
+#pragma once
+
+#include "Behaviour/RCBehaviour.h"
+#include "AvaSceneStateRCEventBehavior.generated.h"
+
+UCLASS(MinimalAPI)
+class UAvaSceneStateRCEventBehavior : public URCBehaviour
+{
+	GENERATED_BODY()
+};
diff --git a/Engine/Plugins/Experimental/AvalancheSceneState/Source/AvalancheSceneState/Public/RemoteControl/AvaSceneStateRCEventBehaviorNode.h b/Engine/Plugins/Experimental/AvalancheSceneState/Source/AvalancheSceneState/Public/RemoteControl/AvaSceneStateRCEventBehaviorNode.h
new file mode 100644
--- /dev/null
+++ b/Engine/Plugins/Experimental/AvalancheSceneState/Source/AvalancheSceneState/Public/RemoteControl/AvaSceneStateRCEventBehaviorNode.h
@@ -0,0 +1,28 @@
+// Copyright Epic Games, Inc. All Rights Reserved.
+
+#pragma once
+
+#include "Behaviour/RCBehaviourNode.h"
+#include "SceneStateEventTemplate.h"
+#include "AvaSceneStateRCEventBehaviorNode.generated.h"
+
+UCLASS(MinimalAPI)
+class UAvaSceneStateRCEventBehaviorNode : public URCBehaviourNode
+{
+	GENERATED_BODY()
+
+public:
+	UAvaSceneStateRCEventBehaviorNode();
+
+protected:
+	//~ Begin URCBehaviourNode
+	bool Execute(URCBehaviour* InBehavior) const override;
+	bool IsSupported(URCBehaviour* InBehavior) const override;
+	void OnPassed(URCBehaviour* InBehavior) const;
+	UClass* GetBehaviourClass() const override;
+	//~ End URCBehaviourNode
+
+private:
+	UPROPERTY(EditAnywhere, Category="Scene State")
+	FSceneStateEventTemplate Event;
+};
diff --git a/Engine/Plugins/Experimental/AvalancheSceneState/Source/AvalancheSceneState/Public/RemoteControl/AvaSceneStateRCTask.h b/Engine/Plugins/Experimental/AvalancheSceneState/Source/AvalancheSceneState/Public/RemoteControl/AvaSceneStateRCTask.h
new file mode 100644
--- /dev/null
+++ b/Engine/Plugins/Experimental/AvalancheSceneState/Source/AvalancheSceneState/Public/RemoteControl/AvaSceneStateRCTask.h
@@ -0,0 +1,81 @@
+// Copyright Epic Games, Inc. All Rights Reserved.
+
+#pragma once
+
+#include "AvaRCControllerId.h"
+#include "AvaSceneStateRCTaskBinding.h"
+#include "Misc/Guid.h"
+#include "StructUtils/PropertyBag.h"
+#include "Tasks/SceneStateTask.h"
+#include "Tasks/SceneStateTaskBindingExtension.h"
+#include "Tasks/SceneStateTaskInstance.h"
+#include "AvaSceneStateRCTask.generated.h"
+
+class URemoteControlPreset;
+
+USTRUCT()
+struct FAvaSceneStateRCControllerMapping
+{
+	GENERATED_BODY()
+
+	bool operator==(const FGuid& InPropertyId) const
+	{
+		return SourcePropertyId == InPropertyId;
+	}
+
+	/** Id to the controller to copy the value to */
+	UPROPERTY(EditAnywhere, Category="Scene State")
+	FAvaRCControllerId TargetController;
+
+	/** Id of the property (in property bag) whose value will be copied into the target controller */
+	UPROPERTY(EditAnywhere, Category="Scene State")
+	FGuid SourcePropertyId;
+};
+
+/**
+ * Instance Data for the RC Task
+ * @see FAvaSceneStateRCTask
+ */
+USTRUCT()
+struct FAvaSceneStateRCTaskInstance : public FSceneStateTaskInstance
+{
+	GENERATED_BODY()
+
+#if WITH_EDITORONLY_DATA
+	/** Struct Id uniquely identifying ControllerValues */
+	UPROPERTY(VisibleAnywhere, Category="Scene State", meta=(NoBinding))
+	FGuid ControllerValuesId;
+#endif
+
+	UPROPERTY(EditAnywhere, Category="Scene State")
+	TArray<FAvaSceneStateRCControllerMapping> ControllerMappings;
+
+	UPROPERTY(EditAnywhere, Category="Scene State")
+	FInstancedPropertyBag ControllerValues;
+};
+
+USTRUCT(DisplayName="Set RC Controller Values", Category="Motion Design")
+struct FAvaSceneStateRCTask : public FSceneStateTask
+{
+	GENERATED_BODY()
+
+	using FInstanceDataType = FAvaSceneStateRCTaskInstance;
+
+	FAvaSceneStateRCTask();
+
+protected:
+	//~ Begin FSceneStateTask
+#if WITH_EDITOR
+	virtual const UScriptStruct* OnGetTaskInstanceType() const override;
+	virtual void OnBuildTaskInstance(UObject* InOuter, FStructView InTaskInstance) const override;
+#endif
+	virtual const FSceneStateTaskBindingExtension* OnGetBindingExtension() const override;
+	virtual void OnStart(const FSceneStateExecutionContext& InContext, FStructView InTaskInstance) const override;
+	//~ End FSceneStateTask
+
+private:
+	bool GetControllerDataView(const FSceneStateExecutionContext& InContext, FStructView& OutControllerView, URemoteControlPreset*& OutPreset) const;
+
+	UPROPERTY()
+	FAvaSceneStateRCTaskBinding Binding;
+};
diff --git a/Engine/Plugins/Experimental/AvalancheSceneState/Source/AvalancheSceneState/Public/RemoteControl/AvaSceneStateRCTaskBinding.h b/Engine/Plugins/Experimental/AvalancheSceneState/Source/AvalancheSceneState/Public/RemoteControl/AvaSceneStateRCTaskBinding.h
new file mode 100644
--- /dev/null
+++ b/Engine/Plugins/Experimental/AvalancheSceneState/Source/AvalancheSceneState/Public/RemoteControl/AvaSceneStateRCTaskBinding.h
@@ -0,0 +1,34 @@
+// Copyright Epic Games, Inc. All Rights Reserved.
+
+#pragma once
+
+#include "Tasks/SceneStateTaskBindingExtension.h"
+#include "AvaSceneStateRCTaskBinding.generated.h"
+
+struct FAvaSceneStateRCTaskInstance;
+
+USTRUCT()
+struct FAvaSceneStateRCTaskBinding : public FSceneStateTaskBindingExtension
+{
+	GENERATED_BODY()
+
+	using FInstanceDataType = FAvaSceneStateRCTaskInstance;
+
+	/** Data View Indices */
+	static constexpr uint16 CONTROLLER_VALUES_DATA_INDEX = 0;
+
+protected:
+	//~ Begin FSceneStateTaskBindingExtension
+#if WITH_EDITOR
+	virtual void VisitBindingDescs(FConstStructView InTaskInstance, TFunctionRef<void(const UE::SceneState::FTaskBindingDesc&)> InFunctor) const override;
+	virtual void SetBindingBatch(uint16 InDataIndex, uint16 InBatchIndex) override;
+	virtual bool FindDataById(FStructView InTaskInstance, const FGuid& InStructId, FStructView& OutDataView, uint16& OutDataIndex) const override;
+#endif
+	virtual bool FindDataByIndex(FStructView InTaskInstance, uint16 InDataIndex, FStructView& OutDataView) const override;
+	virtual void VisitBindingBatches(FStructView InTaskInstance, TFunctionRef<void(uint16, FStructView)> InFunctor) const override;
+	//~ End FSceneStateTaskBindingExtension
+
+private:
+	UPROPERTY()
+	uint16 ControllerValuesBatchIndex = std::numeric_limits<uint16>::max();
+};
diff --git a/Engine/Plugins/Experimental/AvalancheSceneState/Source/AvalancheSceneState/Public/Sequence/AvaSceneStatePlaySequenceTask.h b/Engine/Plugins/Experimental/AvalancheSceneState/Source/AvalancheSceneState/Public/Sequence/AvaSceneStatePlaySequenceTask.h
new file mode 100644
--- /dev/null
+++ b/Engine/Plugins/Experimental/AvalancheSceneState/Source/AvalancheSceneState/Public/Sequence/AvaSceneStatePlaySequenceTask.h
@@ -0,0 +1,70 @@
+// Copyright Epic Games, Inc. All Rights Reserved.
+
+#pragma once
+
+#include "AvaSequenceName.h"
+#include "AvaSequenceShared.h"
+#include "AvaTagHandle.h"
+#include "Tasks/SceneStateTask.h"
+#include "Tasks/SceneStateTaskInstance.h"
+#include "Transition/AvaTransitionSequenceEnums.h"
+#include "AvaSceneStatePlaySequenceTask.generated.h"
+
+class IAvaSequencePlaybackObject;
+class UAvaSequencePlayer;
+
+namespace UE::SceneState
+{
+	struct FTaskExecutionContext;
+}
+
+USTRUCT()
+struct FAvaSceneStatePlaySequenceTaskInstance : public FSceneStateTaskInstance
+{
+	GENERATED_BODY()
+
+	/** The method to find the sequence to play */
+	UPROPERTY(EditAnywhere, Category="Sequence")
+	EAvaTransitionSequenceQueryType SequenceQueryType = EAvaTransitionSequenceQueryType::Name;
+
+	/** The name of the sequences to play (if query type is set to name) */
+	UPROPERTY(EditAnywhere, Category="Sequence", meta=(EditCondition="SequenceQueryType==EAvaTransitionSequenceQueryType::Name", EditConditionHides))
+	FAvaSequenceName SequenceName;
+
+	/** The tag of the sequences to play (if query type is set to tag) */
+	UPROPERTY(EditAnywhere, Category="Sequence", meta=(EditCondition="SequenceQueryType==EAvaTransitionSequenceQueryType::Tag", EditConditionHides))
+	FAvaTagHandle SequenceTag;
+
+	UPROPERTY(EditAnywhere, Category="Sequence")
+	FAvaSequencePlayParams PlaySettings;
+
+	/** The wait type before this task completes */
+	UPROPERTY(EditAnywhere, Category="Sequencer")
+	EAvaTransitionSequenceWaitType WaitType = EAvaTransitionSequenceWaitType::WaitUntilStop;
+
+	/** Active sequence players on this instance */
+	UPROPERTY(Transient)
+	TArray<TObjectPtr<UAvaSequencePlayer>> SequencePlayers;
+
+	FDelegateHandle OnSequenceFinishedHandle;
+};
+
+USTRUCT(DisplayName="Play Sequence", Category="Motion Design", meta=(ToolTip="Plays a Motion Design Sequence"))
+struct FAvaSceneStatePlaySequenceTask : public FSceneStateTask
+{
+	GENERATED_BODY()
+
+	using FInstanceDataType = FAvaSceneStatePlaySequenceTaskInstance;
+
+protected:
+	//~ Begin FSceneStateTask
+#if WITH_EDITOR
+	virtual const UScriptStruct* OnGetTaskInstanceType() const override;
+#endif
+	virtual void OnStart(const FSceneStateExecutionContext& InContext, FStructView InTaskInstance) const override;
+	virtual void OnStop(const FSceneStateExecutionContext& InContext, FStructView InTaskInstance, ESceneStateTaskStopReason InStopReason) const override;
+	//~ End FSceneStateTask
+
+private:
+	static void OnSequenceStopped(UAvaSequencePlayer* InPlayer, UAvaSequence* InSequence, UE::SceneState::FTaskExecutionContext InTaskContext);
+};
diff --git a/Engine/Plugins/Experimental/AvalancheSceneState/Source/AvalancheSceneStateBlueprint/AvalancheSceneStateBlueprint.Build.cs b/Engine/Plugins/Experimental/AvalancheSceneState/Source/AvalancheSceneStateBlueprint/AvalancheSceneStateBlueprint.Build.cs
new file mode 100644
--- /dev/null
+++ b/Engine/Plugins/Experimental/AvalancheSceneState/Source/AvalancheSceneStateBlueprint/AvalancheSceneStateBlueprint.Build.cs
@@ -0,0 +1,29 @@
+// Copyright Epic Games, Inc. All Rights Reserved.
+
+using UnrealBuildTool;
+
+public class AvalancheSceneStateBlueprint : ModuleRules
+{
+    public AvalancheSceneStateBlueprint(ReadOnlyTargetRules Target) : base(Target)
+    {
+        PCHUsage = ModuleRules.PCHUsageMode.UseExplicitOrSharedPCHs;
+
+        PublicDependencyModuleNames.AddRange(
+            new string[]
+            {
+                "Core",
+                "Engine",
+                "PropertyBindingUtils",
+                "SceneStateBlueprint",
+            }
+        );
+
+        PrivateDependencyModuleNames.AddRange(
+            new string[]
+            {
+                "AvalancheSceneState",
+                "CoreUObject",
+            }
+        );
+    }
+}
\ No newline at end of file
diff --git a/Engine/Plugins/Experimental/AvalancheSceneState/Source/AvalancheSceneStateBlueprint/Private/AvaSceneStateBlueprintModule.cpp b/Engine/Plugins/Experimental/AvalancheSceneState/Source/AvalancheSceneStateBlueprint/Private/AvaSceneStateBlueprintModule.cpp
new file mode 100644
--- /dev/null
+++ b/Engine/Plugins/Experimental/AvalancheSceneState/Source/AvalancheSceneStateBlueprint/Private/AvaSceneStateBlueprintModule.cpp
@@ -0,0 +1,5 @@
+// Copyright Epic Games, Inc. All Rights Reserved.
+
+#include "Modules/ModuleManager.h"
+
+IMPLEMENT_MODULE(FDefaultModuleImpl, AvalancheSceneStateBlueprint);
diff --git a/Engine/Plugins/Experimental/AvalancheSceneState/Source/AvalancheSceneStateBlueprint/Public/AvaSceneStateBlueprint.h b/Engine/Plugins/Experimental/AvalancheSceneState/Source/AvalancheSceneStateBlueprint/Public/AvaSceneStateBlueprint.h
new file mode 100644
--- /dev/null
+++ b/Engine/Plugins/Experimental/AvalancheSceneState/Source/AvalancheSceneStateBlueprint/Public/AvaSceneStateBlueprint.h
@@ -0,0 +1,16 @@
+// Copyright Epic Games, Inc. All Rights Reserved.
+
+#pragma once
+
+#include "SceneStateBlueprint.h"
+#include "AvaSceneStateBlueprint.generated.h"
+
+/**
+ * Scene State Blueprint child that overrides the Property Binder class with a child class
+ * @see UAvaSceneStateBinder
+ */
+UCLASS(MinimalAPI, DisplayName="Motion Design Scene State Blueprint")
+class UAvaSceneStateBlueprint : public USceneStateBlueprint
+{
+	GENERATED_BODY()
+};
diff --git a/Engine/Plugins/Experimental/AvalancheSceneState/Source/AvalancheSceneStateEditor/AvalancheSceneStateEditor.Build.cs b/Engine/Plugins/Experimental/AvalancheSceneState/Source/AvalancheSceneStateEditor/AvalancheSceneStateEditor.Build.cs
new file mode 100644
--- /dev/null
+++ b/Engine/Plugins/Experimental/AvalancheSceneState/Source/AvalancheSceneStateEditor/AvalancheSceneStateEditor.Build.cs
@@ -0,0 +1,36 @@
+// Copyright Epic Games, Inc. All Rights Reserved.
+
+using UnrealBuildTool;
+
+public class AvalancheSceneStateEditor : ModuleRules
+{
+    public AvalancheSceneStateEditor(ReadOnlyTargetRules Target) : base(Target)
+    {
+        PCHUsage = ModuleRules.PCHUsageMode.UseExplicitOrSharedPCHs;
+
+        PrivateDependencyModuleNames.AddRange(
+            new string[]
+            {
+                "AvalancheEditorCore",
+                "AvalancheSceneState",
+                "AvalancheSceneStateBlueprint",
+                "BlueprintGraph",
+                "Core",
+                "CoreUObject",
+                "Engine",
+                "KismetWidgets",
+                "Projects",
+                "PropertyEditor",
+                "RemoteControlLogic",
+                "SceneStateBlueprint",
+                "SceneStateBlueprintEditor",
+                "SceneStateGameplay",
+                "Slate",
+                "SlateCore",
+                "StructUtilsEditor",
+                "ToolMenus",
+                "UnrealEd",
+            }
+        );
+    }
+}
\ No newline at end of file
diff --git a/Engine/Plugins/Experimental/AvalancheSceneState/Source/AvalancheSceneStateEditor/Private/AvaSceneStateEditorModule.cpp b/Engine/Plugins/Experimental/AvalancheSceneState/Source/AvalancheSceneStateEditor/Private/AvaSceneStateEditorModule.cpp
new file mode 100644
--- /dev/null
+++ b/Engine/Plugins/Experimental/AvalancheSceneState/Source/AvalancheSceneStateEditor/Private/AvaSceneStateEditorModule.cpp
@@ -0,0 +1,58 @@
+// Copyright Epic Games, Inc. All Rights Reserved.
+
+#include "AvaSceneStateEditorModule.h"
+#include "AvaEditorBuilder.h"
+#include "AvaSceneStateBlueprint.h"
+#include "AvaSceneStateEditorStyle.h"
+#include "AvaSceneStateExtension.h"
+#include "DetailsView/AvaSceneStateRCTaskDetails.h"
+#include "ISceneStateBlueprintEditorModule.h"
+#include "Modules/ModuleManager.h"
+#include "PropertyEditorModule.h"
+#include "RemoteControl/AvaSceneStateRCEventBehaviorNode.h"
+#include "RemoteControl/AvaSceneStateRCTask.h"
+
+IMPLEMENT_MODULE(FAvaSceneStateEditorModule, AvalancheSceneStateEditor)
+
+void FAvaSceneStateEditorModule::StartupModule()
+{
+	UE::AvaSceneState::Editor::FEditorStyle::Get();
+
+	UE::SceneState::Editor::IBlueprintEditorModule::Get().RegisterCompiler(UAvaSceneStateBlueprint::StaticClass());
+
+	OnEditorBuildHandle = FAvaEditorBuilder::OnEditorBuild.AddLambda(
+		[](FAvaEditorBuilder& InBuilder)
+		{
+			InBuilder.AddExtension<FAvaSceneStateExtension>();	
+		});
+
+	RegisterCustomizations();
+}
+
+void FAvaSceneStateEditorModule::ShutdownModule()
+{
+	FAvaEditorBuilder::OnEditorBuild.Remove(OnEditorBuildHandle);
+	OnEditorBuildHandle.Reset();
+
+	UnregisterCustomizations();
+}
+
+void FAvaSceneStateEditorModule::RegisterCustomizations()
+{
+	FPropertyEditorModule& PropertyEditorModule = FModuleManager::LoadModuleChecked<FPropertyEditorModule>("PropertyEditor");
+
+	PropertyEditorModule.RegisterCustomPropertyTypeLayout(CustomizedTypes.Add_GetRef(FAvaSceneStateRCTaskInstance::StaticStruct()->GetFName())
+		, FOnGetPropertyTypeCustomizationInstance::CreateStatic(&FAvaSceneStateRCTaskDetails::MakeInstance));
+}
+
+void FAvaSceneStateEditorModule::UnregisterCustomizations()
+{
+	if (FPropertyEditorModule* PropertyEditorModule = FModuleManager::GetModulePtr<FPropertyEditorModule>("PropertyEditor"))
+	{
+		for (FName CustomizedType : CustomizedTypes)
+		{
+			PropertyEditorModule->UnregisterCustomPropertyTypeLayout(CustomizedType);	
+		}
+		CustomizedTypes.Reset();
+	}
+}
diff --git a/Engine/Plugins/Experimental/AvalancheSceneState/Source/AvalancheSceneStateEditor/Private/AvaSceneStateEditorModule.h b/Engine/Plugins/Experimental/AvalancheSceneState/Source/AvalancheSceneStateEditor/Private/AvaSceneStateEditorModule.h
new file mode 100644
--- /dev/null
+++ b/Engine/Plugins/Experimental/AvalancheSceneState/Source/AvalancheSceneStateEditor/Private/AvaSceneStateEditorModule.h
@@ -0,0 +1,23 @@
+// Copyright Epic Games, Inc. All Rights Reserved.
+
+#pragma once
+
+#include "Containers/Array.h"
+#include "Delegates/IDelegateInstance.h"
+#include "Modules/ModuleInterface.h"
+#include "UObject/NameTypes.h"
+
+class FAvaSceneStateEditorModule : public IModuleInterface
+{
+	//~ Begin IModuleInterface
+	virtual void StartupModule() override;
+	virtual void ShutdownModule() override;
+	//~ End IModuleInterface
+
+	void RegisterCustomizations();
+	void UnregisterCustomizations();
+
+	FDelegateHandle OnEditorBuildHandle;
+
+	TArray<FName> CustomizedTypes;
+};
diff --git a/Engine/Plugins/Experimental/AvalancheSceneState/Source/AvalancheSceneStateEditor/Private/AvaSceneStateEditorStyle.cpp b/Engine/Plugins/Experimental/AvalancheSceneState/Source/AvalancheSceneStateEditor/Private/AvaSceneStateEditorStyle.cpp
new file mode 100644
--- /dev/null
+++ b/Engine/Plugins/Experimental/AvalancheSceneState/Source/AvalancheSceneStateEditor/Private/AvaSceneStateEditorStyle.cpp
@@ -0,0 +1,30 @@
+// Copyright Epic Games, Inc. All Rights Reserved.
+
+#include "AvaSceneStateEditorStyle.h"
+#include "Interfaces/IPluginManager.h"
+#include "Styling/CoreStyle.h"
+#include "Styling/SlateStyleMacros.h"
+#include "Styling/SlateStyleRegistry.h"
+
+namespace UE::AvaSceneState::Editor
+{
+
+FEditorStyle::FEditorStyle()
+	: FSlateStyleSet(TEXT("AvaSceneStateEditor"))
+{
+	ContentRootDir = IPluginManager::Get().FindPlugin(UE_PLUGIN_NAME)->GetContentDir();
+	CoreContentRootDir = FPaths::EngineContentDir() / TEXT("Slate");
+
+	// Task Icons
+	Set("TaskIcon.AvaSceneStatePlaySequenceTask", new IMAGE_BRUSH_SVG("Icons/Tasks/PlaySequence", CoreStyleConstants::Icon16x16));
+	Set("TaskIcon.AvaSceneStateRCTask", new IMAGE_BRUSH("Icons/Tasks/RemoteControl", CoreStyleConstants::Icon16x16));
+
+	FSlateStyleRegistry::RegisterSlateStyle(*this);
+}
+
+FEditorStyle::~FEditorStyle()
+{
+	FSlateStyleRegistry::UnRegisterSlateStyle(*this);
+}
+
+}
diff --git a/Engine/Plugins/Experimental/AvalancheSceneState/Source/AvalancheSceneStateEditor/Private/AvaSceneStateEditorStyle.h b/Engine/Plugins/Experimental/AvalancheSceneState/Source/AvalancheSceneStateEditor/Private/AvaSceneStateEditorStyle.h
new file mode 100644
--- /dev/null
+++ b/Engine/Plugins/Experimental/AvalancheSceneState/Source/AvalancheSceneStateEditor/Private/AvaSceneStateEditorStyle.h
@@ -0,0 +1,23 @@
+// Copyright Epic Games, Inc. All Rights Reserved.
+
+#pragma once
+
+#include "Styling/SlateStyle.h"
+
+namespace UE::AvaSceneState::Editor
+{
+
+class FEditorStyle : public FSlateStyleSet
+{
+public:
+	static FEditorStyle& Get()
+	{
+		static FEditorStyle Instance;
+		return Instance;
+	}
+
+	FEditorStyle();
+	virtual ~FEditorStyle() override;
+};
+
+} // UE::AvaSceneState::Editor
diff --git a/Engine/Plugins/Experimental/AvalancheSceneState/Source/AvalancheSceneStateEditor/Private/AvaSceneStateExtension.cpp b/Engine/Plugins/Experimental/AvalancheSceneState/Source/AvalancheSceneStateEditor/Private/AvaSceneStateExtension.cpp
new file mode 100644
--- /dev/null
+++ b/Engine/Plugins/Experimental/AvalancheSceneState/Source/AvalancheSceneStateEditor/Private/AvaSceneStateExtension.cpp
@@ -0,0 +1,115 @@
+// Copyright Epic Games, Inc. All Rights Reserved.
+
+#include "AvaSceneStateExtension.h"
+#include "AvaSceneStateActor.h"
+#include "AvaSceneStateBlueprint.h"
+#include "BlueprintActionDatabase.h"
+#include "Editor.h"
+#include "Engine/World.h"
+#include "EngineUtils.h"
+#include "SceneStateActor.h"
+#include "SceneStateBlueprintFactory.h"
+#include "SceneStateObject.h"
+#include "Styling/SlateIconFinder.h"
+#include "Subsystems/AssetEditorSubsystem.h"
+#include "ToolMenu.h"
+#include "ToolMenuEntry.h"
+#include "ToolMenuSection.h"
+
+#define LOCTEXT_NAMESPACE "AvaSceneStateExtension"
+
+void FAvaSceneStateExtension::ExtendToolbarMenu(UToolMenu& InMenu)
+{
+	FToolMenuSection& Section = InMenu.FindOrAddSection(DefaultSectionName);
+
+	FToolMenuEntry& Entry = Section.AddEntry(FToolMenuEntry::InitToolBarButton(TEXT("SceneStateButton")
+		, FExecuteAction::CreateSP(this, &FAvaSceneStateExtension::OpenSceneStateBlueprintEditor)
+		, LOCTEXT("SceneStateLabel", "Scene State")
+		, LOCTEXT("SceneStateTooltip", "Opens the Scene State Editor for the given Motion Design Scene")
+		, FSlateIconFinder::FindCustomIconForClass(UAvaSceneStateBlueprint::StaticClass(), TEXT("ClassThumbnail"))));
+
+	Entry.StyleNameOverride = TEXT("CalloutToolbar");
+}
+
+void FAvaSceneStateExtension::Cleanup()
+{
+	UWorld* const World = GetWorld();
+	if (!World)
+	{
+		return;
+	}
+
+	FBlueprintActionDatabase* BlueprintActionDatabase = FBlueprintActionDatabase::TryGet();
+	if (!BlueprintActionDatabase)
+	{
+		return;
+	}
+
+	for (AAvaSceneStateActor* SceneStateActor : TActorRange<AAvaSceneStateActor>(World))
+	{
+		if (SceneStateActor->SceneStateBlueprint)
+		{
+			BlueprintActionDatabase->ClearAssetActions(SceneStateActor->SceneStateBlueprint.Get());
+		}
+	}
+}
+
+AAvaSceneStateActor* FAvaSceneStateExtension::FindOrSpawnActor() const
+{
+	UWorld* const World = GetWorld();
+	if (!World)
+	{
+		return nullptr;
+	}
+
+	for (AAvaSceneStateActor* SceneStateActor : TActorRange<AAvaSceneStateActor>(World))
+	{
+		return SceneStateActor;
+	}
+
+	return World->SpawnActor<AAvaSceneStateActor>();
+}
+
+USceneStateBlueprint* FAvaSceneStateExtension::CreateSceneStateBlueprint(AAvaSceneStateActor* InSceneStateActor) const
+{
+	check(InSceneStateActor);
+
+	UFactory* SceneStateBlueprintFactory = NewObject<USceneStateBlueprintFactory>(GetTransientPackage());
+	check(SceneStateBlueprintFactory);
+
+	USceneStateBlueprint* NewSceneStateBlueprint = CastChecked<USceneStateBlueprint>(SceneStateBlueprintFactory->FactoryCreateNew(UAvaSceneStateBlueprint::StaticClass()
+		, InSceneStateActor
+		, TEXT("SceneStateBlueprint")
+		, RF_Transactional
+		, nullptr
+		, GWarn));
+
+	// Clear Standalone flags as this Blueprint will be under the Scene State Actor
+	NewSceneStateBlueprint->ClearFlags(RF_Standalone);
+
+	InSceneStateActor->SetSceneStateBlueprint(NewSceneStateBlueprint);
+	InSceneStateActor->UpdateSceneStateClass();
+
+	return NewSceneStateBlueprint;
+}
+
+void FAvaSceneStateExtension::OpenSceneStateBlueprintEditor()
+{
+	AAvaSceneStateActor* SceneStateActor = FindOrSpawnActor();
+	if (!SceneStateActor)
+	{
+		return;
+	}
+
+	if (!SceneStateActor->SceneStateBlueprint)
+	{
+		CreateSceneStateBlueprint(SceneStateActor);
+	}
+	check(SceneStateActor->SceneStateBlueprint);
+
+	UAssetEditorSubsystem* AssetEditorSubsystem = GEditor->GetEditorSubsystem<UAssetEditorSubsystem>();
+	check(AssetEditorSubsystem);
+	AssetEditorSubsystem->OpenEditorForAsset(SceneStateActor->SceneStateBlueprint);
+}
+
+#undef LOCTEXT_NAMESPACE
diff --git a/Engine/Plugins/Experimental/AvalancheSceneState/Source/AvalancheSceneStateEditor/Private/AvaSceneStateExtension.h b/Engine/Plugins/Experimental/AvalancheSceneState/Source/AvalancheSceneStateEditor/Private/AvaSceneStateExtension.h
new file mode 100644
--- /dev/null
+++ b/Engine/Plugins/Experimental/AvalancheSceneState/Source/AvalancheSceneStateEditor/Private/AvaSceneStateExtension.h
@@ -0,0 +1,26 @@
+// Copyright Epic Games, Inc. All Rights Reserved.
+
+#pragma once
+
+#include "IAvaEditorExtension.h"
+
+class AAvaSceneStateActor;
+class USceneStateBlueprint;
+
+class FAvaSceneStateExtension : public FAvaEditorExtension
+{
+public:
+	UE_AVA_INHERITS(FAvaSceneStateExtension, FAvaEditorExtension);
+
+	//~ Begin IAvaEditorExtension
+	virtual void ExtendToolbarMenu(UToolMenu& InMenu) override;
+	virtual void Cleanup() override;
+	//~ End IAvaEditorExtension
+
+private:
+	AAvaSceneStateActor* FindOrSpawnActor() const;
+
+	USceneStateBlueprint* CreateSceneStateBlueprint(AAvaSceneStateActor* InSceneStateActor) const;
+
+	void OpenSceneStateBlueprintEditor();
+};
diff --git a/Engine/Plugins/Experimental/AvalancheSceneState/Source/AvalancheSceneStateEditor/Private/DetailsView/AvaSceneStateRCTaskDetails.cpp b/Engine/Plugins/Experimental/AvalancheSceneState/Source/AvalancheSceneStateEditor/Private/DetailsView/AvaSceneStateRCTaskDetails.cpp
new file mode 100644
--- /dev/null
+++ b/Engine/Plugins/Experimental/AvalancheSceneState/Source/AvalancheSceneStateEditor/Private/DetailsView/AvaSceneStateRCTaskDetails.cpp
@@ -0,0 +1,30 @@
+// Copyright Epic Games, Inc. All Rights Reserved.
+
+#include "AvaSceneStateRCTaskDetails.h"
+#include "AvaSceneStateRCValuesDetails.h"
+#include "IDetailChildrenBuilder.h"
+#include "PropertyBagDetails.h"
+#include "PropertyHandle.h"
+#include "RemoteControl/AvaSceneStateRCTask.h"
+#include "SceneStateBlueprintEditorUtils.h"
+
+void FAvaSceneStateRCTaskDetails::CustomizeHeader(TSharedRef<IPropertyHandle> InPropertyHandle
+	, FDetailWidgetRow& InHeaderRow
+	, IPropertyTypeCustomizationUtils& InCustomizationUtils)
+{
+}
+
+void FAvaSceneStateRCTaskDetails::CustomizeChildren(TSharedRef<IPropertyHandle> InPropertyHandle
+	, IDetailChildrenBuilder& InChildBuilder
+	, IPropertyTypeCustomizationUtils& InCustomizationUtils)
+{
+	const TSharedRef<IPropertyHandle> ControllerValuesIdHandle = InPropertyHandle->GetChildHandle(GET_MEMBER_NAME_CHECKED(FAvaSceneStateRCTaskInstance, ControllerValuesId)).ToSharedRef();
+	ControllerValuesIdHandle->MarkHiddenByCustomization();
+
+	// If the id could not be gotten, continue displaying the customization as normal.
+	// Not having an id only means that these rows will not be bindable
+	FGuid ControllerValuesId;
+	UE::SceneState::Editor::GetGuid(ControllerValuesIdHandle, ControllerValuesId);
+
+	InChildBuilder.AddCustomBuilder(MakeShared<FAvaSceneStateRCValuesDetails>(InPropertyHandle, ControllerValuesId, InCustomizationUtils.GetPropertyUtilities()));
+}
diff --git a/Engine/Plugins/Experimental/AvalancheSceneState/Source/AvalancheSceneStateEditor/Private/DetailsView/AvaSceneStateRCTaskDetails.h b/Engine/Plugins/Experimental/AvalancheSceneState/Source/AvalancheSceneStateEditor/Private/DetailsView/AvaSceneStateRCTaskDetails.h
new file mode 100644
--- /dev/null
+++ b/Engine/Plugins/Experimental/AvalancheSceneState/Source/AvalancheSceneStateEditor/Private/DetailsView/AvaSceneStateRCTaskDetails.h
@@ -0,0 +1,20 @@
+// Copyright Epic Games, Inc. All Rights Reserved.
+
+#pragma once
+
+#include "IPropertyTypeCustomization.h"
+
+class FAvaSceneStateRCTaskDetails : public IPropertyTypeCustomization
+{
+public:
+	static TSharedRef<IPropertyTypeCustomization> MakeInstance()
+	{
+		return MakeShared<FAvaSceneStateRCTaskDetails>();
+	}
+
+protected:
+	//~ Begin IPropertyTypeCustomization
+	virtual void CustomizeHeader(TSharedRef<IPropertyHandle> InPropertyHandle, FDetailWidgetRow& InHeaderRow, IPropertyTypeCustomizationUtils& InCustomizationUtils) override;
+	virtual void CustomizeChildren(TSharedRef<IPropertyHandle> InPropertyHandle, IDetailChildrenBuilder& InChildBuilder, IPropertyTypeCustomizationUtils& InCustomizationUtils) override;
+	//~ End IPropertyTypeCustomization
+};
diff --git a/Engine/Plugins/Experimental/AvalancheSceneState/Source/AvalancheSceneStateEditor/Private/DetailsView/AvaSceneStateRCUtils.cpp b/Engine/Plugins/Experimental/AvalancheSceneState/Source/AvalancheSceneStateEditor/Private/DetailsView/AvaSceneStateRCUtils.cpp
new file mode 100644
--- /dev/null
+++ b/Engine/Plugins/Experimental/AvalancheSceneState/Source/AvalancheSceneStateEditor/Private/DetailsView/AvaSceneStateRCUtils.cpp
@@ -0,0 +1,168 @@
+// Copyright Epic Games, Inc. All Rights Reserved.
+
+#include "AvaSceneStateRCUtils.h"
+#include "EdGraph/EdGraphPin.h"
+#include "Misc/EnumerateRange.h"
+#include "PropertyBagDetails.h"
+#include "PropertyHandle.h"
+#include "RemoteControl/AvaSceneStateRCTask.h"
+#include "RemoteControlLogicConfig.h"
+
+namespace UE::AvaSceneStateRCUtils
+{
+
+void ForEachControllerSupportedType(TFunctionRef<void(const FEdGraphPinType&)> InFunctor)
+{
+	const URemoteControlLogicConfig* RCLogicConfig = GetDefault<URemoteControlLogicConfig>();
+	if (!ensure(RCLogicConfig))
+	{
+		return;
+	}
+
+	for (EPropertyBagPropertyType ControllerType : RCLogicConfig->SupportedControllerTypes)
+	{
+		// Skip structs/objects. These are handled in a separate pass with SupportedControllerStructTypes and SupportedControllerObjectClassPaths
+		if (ControllerType == EPropertyBagPropertyType::Struct || ControllerType == EPropertyBagPropertyType::Object)
+		{
+			continue;
+		}
+
+		InFunctor(UE::StructUtils::GetPropertyDescAsPin(FPropertyBagPropertyDesc(NAME_None, ControllerType)));
+	}
+
+	constexpr const TCHAR* CoreStructTypePathPrefix = TEXT("/Script/CoreUObject.");
+	for (FName ControllerStructType : RCLogicConfig->SupportedControllerStructTypes)
+	{
+		if (UStruct* Struct = FindObject<UScriptStruct>(nullptr, *(CoreStructTypePathPrefix + ControllerStructType.ToString())))
+		{
+			InFunctor(UE::StructUtils::GetPropertyDescAsPin(FPropertyBagPropertyDesc(NAME_None, EPropertyBagPropertyType::Struct, Struct)));
+		}
+	}
+
+	for (FName ControllerObjectType : RCLogicConfig->SupportedControllerObjectClassPaths)
+	{
+		if (UObject* Object = FSoftClassPath(ControllerObjectType.ToString()).TryLoad())
+		{
+			InFunctor(UE::StructUtils::GetPropertyDescAsPin(FPropertyBagPropertyDesc(NAME_None, EPropertyBagPropertyType::Object, Object)));
+		}
+	}
+}
+
+FEdGraphPinType GetPinInfo(const TSharedRef<IPropertyHandle>& InPropertyHandle)
+{
+	if (!InPropertyHandle->IsValidHandle())
+	{
+		return FEdGraphPinType();
+	}
+	if (const FProperty* ValueProperty = InPropertyHandle->GetProperty())
+	{
+		return UE::StructUtils::GetPropertyDescAsPin(FPropertyBagPropertyDesc(NAME_None, ValueProperty));
+	}
+	return FEdGraphPinType();
+}
+
+FPropertyBagPropertyDesc MakeControllerPropertyDesc()
+{
+	FPropertyBagPropertyDesc PropertyDesc;
+	PropertyDesc.ID = FGuid::NewGuid();
+	PropertyDesc.Name = FName(TEXT("Property_") + PropertyDesc.ID.ToString());
+
+	// todo: a controller could have more info on the type to set
+	PropertyDesc.ValueType = EPropertyBagPropertyType::String;
+	return PropertyDesc;
+}
+
+bool SyncPropertyDescs(TArray<FPropertyBagPropertyDesc>& InOutPropertyDescs, TArrayView<FAvaSceneStateRCControllerMapping> InMappings)
+{
+	// Remove all the property descs that are no longer in mapping
+	InOutPropertyDescs.RemoveAll(
+		[&InMappings](const FPropertyBagPropertyDesc& InPropertyDesc)
+		{
+			return !InMappings.FindByKey(InPropertyDesc.ID);
+		});
+
+	InOutPropertyDescs.Reserve(InMappings.Num());
+
+	// property descs cannot have more elements than mappings.
+	// otherwise it means there are mappings that have repeated property desc id
+	if (!ensure(InOutPropertyDescs.Num() <= InMappings.Num()))
+	{
+		return false;
+	}
+
+	// Create new property descs to mappings that do not have a matching desc
+	for (TEnumerateRef<FAvaSceneStateRCControllerMapping> Mapping : EnumerateRange(InMappings))
+	{
+		const int32 MappingIndex = Mapping.GetIndex();
+
+		// No property desc available or mapping is new (and invalid), can add directly
+		if (!InOutPropertyDescs.IsValidIndex(MappingIndex) || !Mapping->SourcePropertyId.IsValid())
+		{
+			// mapping should not have a valid property id here
+			if (!ensure(!Mapping->SourcePropertyId.IsValid()))
+			{
+				return false;
+			}
+
+			// create a new desc at the mapping index and link the mapping to it
+			FPropertyBagPropertyDesc& Desc = InOutPropertyDescs.Insert_GetRef(UE::AvaSceneStateRCUtils::MakeControllerPropertyDesc(), MappingIndex);
+			Mapping->SourcePropertyId = Desc.ID;
+			continue;
+		}
+
+		// mapping should have a valid property at this point (new mapping entries added with invalid ids should've been handled above)
+		if (!ensure(Mapping->SourcePropertyId.IsValid()))
+		{
+			return false;
+		}
+
+		int32 PropertyDescIndex = INDEX_NONE;
+
+		// Find the property desc that matches the valid mapping source id
+		// start from the mapping index as anything before that should already be fixed (i.e. the ids for the elements at each array match)
+		for (int32 Index = MappingIndex; Index < InOutPropertyDescs.Num(); ++Index)
+		{
+			if (InOutPropertyDescs[Index].ID == Mapping->SourcePropertyId)
+			{
+				PropertyDescIndex = Index;
+				break;
+			}
+		}
+
+		// Property index must've been found, as all the property desc entries that didn't have a matching mapping were removed ahead of time
+		if (!ensure(PropertyDescIndex != INDEX_NONE))
+		{
+			return false;
+		}
+
+		// the index matched the current mapping, nothing to do
+		if (PropertyDescIndex == MappingIndex)
+		{
+			continue;
+		}
+
+		// move the property desc to the mapping index
+		FPropertyBagPropertyDesc PropertyDesc = MoveTemp(InOutPropertyDescs[PropertyDescIndex]);
+		InOutPropertyDescs.RemoveAt(PropertyDescIndex);
+		InOutPropertyDescs.Insert(MoveTemp(PropertyDesc), MappingIndex);
+	}
+
+	// validation pass: mapping count must match property desc count now
+	if (!ensure(InOutPropertyDescs.Num() == InMappings.Num()))
+	{
+		return false;
+	}
+
+	// validation pass: all ids must match for each index
+	for (int32 Index = 0; Index < InMappings.Num(); ++Index)
+	{
+		if (!ensure(InMappings[Index].SourcePropertyId == InOutPropertyDescs[Index].ID))
+		{
+			return false;
+		}
+	}
+
+	return true;
+}
+
+} // UE::AvaSceneStateRCUtils
diff --git a/Engine/Plugins/Experimental/AvalancheSceneState/Source/AvalancheSceneStateEditor/Private/DetailsView/AvaSceneStateRCUtils.h b/Engine/Plugins/Experimental/AvalancheSceneState/Source/AvalancheSceneStateEditor/Private/DetailsView/AvaSceneStateRCUtils.h
new file mode 100644
--- /dev/null
+++ b/Engine/Plugins/Experimental/AvalancheSceneState/Source/AvalancheSceneStateEditor/Private/DetailsView/AvaSceneStateRCUtils.h
@@ -0,0 +1,32 @@
+// Copyright Epic Games, Inc. All Rights Reserved.
+
+#pragma once
+
+#include "Containers/ContainersFwd.h"
+#include "Templates/FunctionFwd.h"
+#include "Templates/SharedPointerFwd.h"
+
+class IPropertyHandle;
+class UPropertyBag;
+struct FAvaSceneStateRCControllerMapping;
+struct FEdGraphPinType;
+struct FPropertyBagPropertyDesc;
+
+namespace UE::AvaSceneStateRCUtils
+{
+	/** Calls the given functor for every supported type for Remote Control controllers */
+	void ForEachControllerSupportedType(TFunctionRef<void(const FEdGraphPinType&)> InFunctor);
+
+	/** Gets the pin type for the property under the given property handle */
+	FEdGraphPinType GetPinInfo(const TSharedRef<IPropertyHandle>& InPropertyHandle);
+
+	/** Creates a new generic property desc */
+	FPropertyBagPropertyDesc MakeControllerPropertyDesc();
+
+	/**
+	 * Returns a copy of property descs (some that could already exist in the property bag) that matches the structure in the given mappings
+	 * @param InOutPropertyDescs the array of property descs to sync 
+	 * @param InMappings the mappings to match against
+	 */
+	bool SyncPropertyDescs(TArray<FPropertyBagPropertyDesc>& InOutPropertyDescs, TArrayView<FAvaSceneStateRCControllerMapping> InMappings);
+}
diff --git a/Engine/Plugins/Experimental/AvalancheSceneState/Source/AvalancheSceneStateEditor/Private/DetailsView/AvaSceneStateRCValuesDetails.cpp b/Engine/Plugins/Experimental/AvalancheSceneState/Source/AvalancheSceneStateEditor/Private/DetailsView/AvaSceneStateRCValuesDetails.cpp
new file mode 100644
--- /dev/null
+++ b/Engine/Plugins/Experimental/AvalancheSceneState/Source/AvalancheSceneStateEditor/Private/DetailsView/AvaSceneStateRCValuesDetails.cpp
@@ -0,0 +1,296 @@
+// Copyright Epic Games, Inc. All Rights Reserved.
+
+#include "AvaSceneStateRCValuesDetails.h"
+#include "AvaSceneStateRCUtils.h"
+#include "DetailLayoutBuilder.h"
+#include "DetailWidgetRow.h"
+#include "IDetailChildrenBuilder.h"
+#include "IDetailPropertyRow.h"
+#include "IPropertyUtilities.h"
+#include "IStructureDataProvider.h"
+#include "Misc/EnumerateRange.h"
+#include "RemoteControl/AvaSceneStateRCTask.h"
+#include "SPinTypeSelector.h"
+#include "SceneStateBlueprintEditorUtils.h"
+#include "ScopedTransaction.h"
+
+#define LOCTEXT_NAMESPACE "AvaSceneStateRCValuesDetails"
+
+FAvaSceneStateRCValuesDetails::FAvaSceneStateRCValuesDetails(const TSharedRef<IPropertyHandle>& InStructHandle, const FGuid& InValuesId, TSharedPtr<IPropertyUtilities> InPropUtils)
+	: StructHandle(InStructHandle)
+	, ValuesHandle(InStructHandle->GetChildHandle(GET_MEMBER_NAME_CHECKED(FAvaSceneStateRCTaskInstance, ControllerValues)).ToSharedRef())
+	, MappingsHandle(InStructHandle->GetChildHandle(GET_MEMBER_NAME_CHECKED(FAvaSceneStateRCTaskInstance, ControllerMappings)).ToSharedRef())
+	, PropertyUtilities(InPropUtils)
+	, ValuesId(InValuesId)
+{
+	ValuesHandle->MarkHiddenByCustomization();
+	MappingsHandle->MarkHiddenByCustomization();
+
+	UE::SceneState::Editor::AssignBindingId(MappingsHandle, UE::SceneState::Editor::FindTaskId(InStructHandle));
+	UE::SceneState::Editor::AssignBindingId(ValuesHandle, ValuesId);
+}
+
+void FAvaSceneStateRCValuesDetails::Initialize()
+{
+	bInitialized = true;
+
+	TSharedRef<IPropertyHandleArray> ControllerMappingsArray = MappingsHandle->AsArray().ToSharedRef();
+
+	// NumElementsChanged gets called on RebuildChildren at a time out of a transaction
+	ControllerMappingsArray->SetOnNumElementsChanged(FSimpleDelegate::CreateSP(this, &FAvaSceneStateRCValuesDetails::OnControllerMappingsNumChanged));
+
+	// PropertyValueChanged gets called while still on a transaction
+	// Sync controller values whenever the mappings change
+	MappingsHandle->SetOnPropertyValueChanged(FSimpleDelegate::CreateSP(this, &FAvaSceneStateRCValuesDetails::SyncControllerValues));
+}
+
+void FAvaSceneStateRCValuesDetails::SyncControllerValues()
+{
+	TArray<FAvaSceneStateRCControllerMapping>& Mappings = GetControllerMappings();
+
+	FInstancedPropertyBag& InstancedPropertyBag =  GetInstancedPropertyBag();
+
+	// Copy current layout of property bag struct to mutate into latest layout
+	TArray<FPropertyBagPropertyDesc> PropertyDescs;
+	if (const UPropertyBag* PropertyBag = InstancedPropertyBag.GetPropertyBagStruct())
+	{
+		PropertyDescs = PropertyBag->GetPropertyDescs();
+	}
+
+	// Sync PropertyDescs to match the mapping entries
+	if (!UE::AvaSceneStateRCUtils::SyncPropertyDescs(/*in/out*/PropertyDescs, Mappings))
+	{
+		return;
+	}
+
+	ValuesHandle->NotifyPreChange();
+
+	InstancedPropertyBag.MigrateToNewBagStruct(UPropertyBag::GetOrCreateFromDescs(PropertyDescs));
+
+	ValuesHandle->NotifyPostChange(EPropertyChangeType::ValueSet);
+	ValuesHandle->NotifyFinishedChangingProperties();
+
+	if (PropertyUtilities.IsValid())
+	{
+		PropertyUtilities->RequestForceRefresh();
+	}
+
+	OnRebuildChildren.ExecuteIfBound();
+}
+
+void FAvaSceneStateRCValuesDetails::OnControllerMappingsNumChanged()
+{
+	OnRebuildChildren.ExecuteIfBound();
+}
+
+void FAvaSceneStateRCValuesDetails::ConfigureMappingRow(const TSharedRef<IPropertyHandle>& InMappingHandle, IDetailPropertyRow& InChildRow)
+{
+	FDetailWidgetRow& WidgetRow = InChildRow.CustomWidget(/*bShowChildren*/true);
+	{
+		TSharedPtr<SWidget> NameWidget;
+		TSharedPtr<SWidget> ValueWidget;
+		InChildRow.GetDefaultWidgets(NameWidget, ValueWidget, WidgetRow);
+	}
+
+	WidgetRow
+		.ValueContent()
+		[
+			SNew(SHorizontalBox)
+			+ SHorizontalBox::Slot()
+			[
+				WidgetRow.ValueWidget.Widget
+			]
+			+ SHorizontalBox::Slot()
+			.AutoWidth()
+			.HAlign(HAlign_Left)
+			.VAlign(VAlign_Center)
+			.Padding(4, 1, 0, 1)
+			[
+				InMappingHandle->CreateDefaultPropertyButtonWidgets()
+			]
+		];
+}
+
+void FAvaSceneStateRCValuesDetails::ConfigureValueRow(IDetailPropertyRow& InChildRow)
+{
+	TSharedRef<IPropertyHandle> ValueHandle = InChildRow.GetPropertyHandle().ToSharedRef();
+
+	FDetailWidgetRow& WidgetRow = InChildRow.CustomWidget(/*bShowChildren*/true);
+	{
+		TSharedPtr<SWidget> NameWidget;
+		TSharedPtr<SWidget> ValueWidget;
+		InChildRow.GetDefaultWidgets(NameWidget, ValueWidget, WidgetRow);
+	}
+
+	UE::SceneState::Editor::AssignBindingId(ValueHandle, ValuesId);
+
+	WidgetRow
+		.NameContent()
+		.HAlign(HAlign_Right)
+		[
+			SNew(SBox)
+			.WidthOverride(90)
+			.Padding(0, 0, 4, 0)
+			[
+				SNew(SPinTypeSelector, FGetPinTypeTree::CreateSP(this, &FAvaSceneStateRCValuesDetails::GetControllerSupportedTypes))
+				.TargetPinType_Lambda([ValueHandle]()->FEdGraphPinType
+					{
+						return UE::AvaSceneStateRCUtils::GetPinInfo(ValueHandle);
+					})
+				.SelectorType(SPinTypeSelector::ESelectorType::Partial)
+				.OnPinTypeChanged(this, &FAvaSceneStateRCValuesDetails::OnPinInfoChanged, ValueHandle)
+				.Schema(GetDefault<UPropertyBagSchema>())
+				.bAllowArrays(false)
+				.TypeTreeFilter(ETypeTreeFilter::None)
+				.Font(IDetailLayoutBuilder::GetDetailFont())
+			]
+		];
+}
+
+void FAvaSceneStateRCValuesDetails::OnPinInfoChanged(const FEdGraphPinType& InPinType, TSharedRef<IPropertyHandle> InValueHandle)
+{
+	if (!InValueHandle->IsValidHandle())
+	{
+		return;
+	}
+
+	const FProperty* ValueProperty = InValueHandle->GetProperty();
+	if (!ValueProperty)
+	{
+		return;
+	}
+
+	const UPropertyBag* CurrentPropertyBag = GetInstancedPropertyBag().GetPropertyBagStruct();
+	if (!CurrentPropertyBag)
+	{
+		return;
+	}
+
+	TArray<FPropertyBagPropertyDesc> PropertyDescs(CurrentPropertyBag->GetPropertyDescs());
+
+	FPropertyBagPropertyDesc* const PropertyDesc = PropertyDescs.FindByPredicate(
+		[ValueProperty](const FPropertyBagPropertyDesc& InDesc)
+		{
+			return InDesc.CachedProperty == ValueProperty;
+		});
+
+	if (!PropertyDesc)
+	{
+		return;
+	}
+
+	FScopedTransaction Transaction(LOCTEXT("OnPropertyTypeChanged", "Change Property Type"));
+	ValuesHandle->NotifyPreChange();
+
+	UE::StructUtils::SetPropertyDescFromPin(*PropertyDesc, InPinType);
+	GetInstancedPropertyBag().MigrateToNewBagStruct(UPropertyBag::GetOrCreateFromDescs(PropertyDescs));
+
+	ValuesHandle->NotifyPostChange(EPropertyChangeType::ValueSet);
+	ValuesHandle->NotifyFinishedChangingProperties();
+}
+
+void FAvaSceneStateRCValuesDetails::GetControllerSupportedTypes(TArray<FPinTypeTreeItem>& OutTreeItems, ETypeTreeFilter InTreeFilter)
+{
+	OutTreeItems.Reset();
+
+	const UPropertyBagSchema* PropertyBagSchema = GetDefault<UPropertyBagSchema>();
+	check(PropertyBagSchema);
+
+	UE::AvaSceneStateRCUtils::ForEachControllerSupportedType(
+		[&OutTreeItems, PropertyBagSchema](const FEdGraphPinType& InPinType)
+		{
+			FPinTypeTreeItem TreeItem;
+			if (UObject* SubCategoryObject = InPinType.PinSubCategoryObject.Get())
+			{
+				TreeItem = MakeShared<UEdGraphSchema_K2::FPinTypeTreeInfo>(InPinType.PinCategory
+					, SubCategoryObject
+					, FText::GetEmpty());
+			}
+			else
+			{
+				const FText CategoryText = UPropertyBagSchema::GetCategoryText(InPinType.PinCategory);
+				TreeItem = MakeShared<UEdGraphSchema_K2::FPinTypeTreeInfo>(CategoryText, InPinType.PinCategory, PropertyBagSchema, FText::GetEmpty());
+			}
+
+			TreeItem->SetPinSubTypeCategory(InPinType.PinSubCategory);
+			OutTreeItems.Add(MoveTemp(TreeItem));
+		});
+
+	check(PropertyBagSchema);
+}
+
+FName FAvaSceneStateRCValuesDetails::GetName() const
+{
+	return TEXT("FAvaSceneStateRCValuesDetails");
+}
+
+void FAvaSceneStateRCValuesDetails::GenerateHeaderRowContent(FDetailWidgetRow& InNodeRow)
+{
+	if (!bInitialized)
+	{
+		Initialize();
+	}
+
+	InNodeRow
+		.NameContent()
+		[
+			MappingsHandle->CreatePropertyNameWidget()
+		]
+		.ValueContent()
+		.VAlign(VAlign_Center)
+		[
+			MappingsHandle->CreatePropertyValueWidget(/*bDisplayDefaultPropertyButtons*/true)
+		]
+		.ShouldAutoExpand(true);
+}
+
+void FAvaSceneStateRCValuesDetails::GenerateChildContent(IDetailChildrenBuilder& InChildrenBuilder)
+{
+	// Get the FInstancedStruct property  under the FInstancedPropertyBag
+	TSharedRef<IPropertyHandle> ControllerValuesStructHandle = ValuesHandle->GetChildHandle(TEXT("Value")).ToSharedRef();
+	TSharedRef<IStructureDataProvider> ControllerValuesDataProvider = UE::SceneState::Editor::CreateInstancedStructDataProvider(ControllerValuesStructHandle);
+
+	ControllerValuesStructHandle->RemoveChildren();
+
+	TSharedRef<IPropertyHandleArray> ControllerMappingsArray = MappingsHandle->AsArray().ToSharedRef();
+
+	uint32 ControllerMappingCount;
+	ControllerMappingsArray->GetNumElements(ControllerMappingCount);
+
+	const TArray<TSharedPtr<IPropertyHandle>> ControllerValueProperties = ControllerValuesStructHandle->AddChildStructure(ControllerValuesDataProvider);
+
+	const uint32 Count = FMath::Min(static_cast<uint32>(ControllerValueProperties.Num()), ControllerMappingCount);
+
+	const FName TargetControllerName = GET_MEMBER_NAME_CHECKED(FAvaSceneStateRCControllerMapping, TargetController);
+
+	// Add two rows: the mapping controller id and the value of it
+	for (uint32 Index = 0; Index < Count; ++Index)
+	{
+		TSharedRef<IPropertyHandle> MappingHandle = ControllerMappingsArray->GetElement(Index);
+		ConfigureMappingRow(MappingHandle, InChildrenBuilder.AddProperty(MappingHandle->GetChildHandle(TargetControllerName).ToSharedRef()));
+		ConfigureValueRow(InChildrenBuilder.AddProperty(ControllerValueProperties[Index].ToSharedRef()));
+	}
+}
+
+void FAvaSceneStateRCValuesDetails::SetOnRebuildChildren(FSimpleDelegate InOnRebuildChildren)
+{
+	OnRebuildChildren = InOnRebuildChildren;
+}
+
+TArray<FAvaSceneStateRCControllerMapping>& FAvaSceneStateRCValuesDetails::GetControllerMappings()
+{
+	TArray<void*> MappingsRawData;
+	MappingsHandle->AccessRawData(MappingsRawData);
+	check(MappingsRawData.Num() == 1);
+	return *static_cast<TArray<FAvaSceneStateRCControllerMapping>*>(MappingsRawData[0]);
+}
+
+FInstancedPropertyBag& FAvaSceneStateRCValuesDetails::GetInstancedPropertyBag()
+{
+	TArray<void*> ValuesRawData;
+	ValuesHandle->AccessRawData(ValuesRawData);
+	check(ValuesRawData.Num() == 1);
+	return *static_cast<FInstancedPropertyBag*>(ValuesRawData[0]);
+}
+
+#undef LOCTEXT_NAMESPACE
diff --git a/Engine/Plugins/Experimental/AvalancheSceneState/Source/AvalancheSceneStateEditor/Private/DetailsView/AvaSceneStateRCValuesDetails.h b/Engine/Plugins/Experimental/AvalancheSceneState/Source/AvalancheSceneStateEditor/Private/DetailsView/AvaSceneStateRCValuesDetails.h
new file mode 100644
--- /dev/null
+++ b/Engine/Plugins/Experimental/AvalancheSceneState/Source/AvalancheSceneStateEditor/Private/DetailsView/AvaSceneStateRCValuesDetails.h
@@ -0,0 +1,78 @@
+// Copyright Epic Games, Inc. All Rights Reserved.
+
+#pragma once
+
+#include "Misc/Guid.h"
+#include "PropertyBagDetails.h"
+#include "PropertyCustomizationHelpers.h"
+#include "SPinTypeSelector.h"
+#include "Templates/SharedPointer.h"
+
+struct FAvaSceneStateRCControllerMapping;
+
+/** Builder to group the controller mapping entries to the instanced property bag value properties */
+class FAvaSceneStateRCValuesDetails : public IDetailCustomNodeBuilder, public TSharedFromThis<FAvaSceneStateRCValuesDetails>
+{
+public:
+	explicit FAvaSceneStateRCValuesDetails(const TSharedRef<IPropertyHandle>& InStructHandle, const FGuid& InValuesId, TSharedPtr<IPropertyUtilities> InPropUtils);
+
+	/** Called once per instance to bind to delegates */
+	void Initialize();
+
+	/**
+	 * Called when mapping has changed, while still during a transaction.
+	 * Syncs the value instanced property bag ensuring there's a valid matching value to every mapping entry.
+	 */
+	void SyncControllerValues();
+
+	/** Called on rebuild time when the mapping num has changed */
+	void OnControllerMappingsNumChanged();
+
+	/** Customizes the controller mapping row */
+	void ConfigureMappingRow(const TSharedRef<IPropertyHandle>& InMappingHandle, IDetailPropertyRow& InChildRow);
+
+	/** Customizes the value row  */
+	void ConfigureValueRow(IDetailPropertyRow& InChildRow);
+
+	/** Called when the type of the value property has changed */
+	void OnPinInfoChanged(const FEdGraphPinType& InPinType, TSharedRef<IPropertyHandle> InValueHandle);
+
+	/** Gathers all the supported types for a remote control preset controller */
+	void GetControllerSupportedTypes(TArray<FPinTypeTreeItem>& OutTreeItems, ETypeTreeFilter InTreeFilter);
+
+	//~ Begin IDetailCustomNodeBuilder
+	virtual FName GetName() const override;
+	virtual bool InitiallyCollapsed() const override { return false; }
+	virtual void GenerateHeaderRowContent(FDetailWidgetRow& InNodeRow) override;
+	virtual void GenerateChildContent(IDetailChildrenBuilder& InChildrenBuilder) override;
+	virtual void SetOnRebuildChildren(FSimpleDelegate InOnRebuildChildren ) override;
+	//~ End IDetailCustomNodeBuilder
+
+private:
+	/** Gets the controller mappings. Expects only one instance to be available */
+	TArray<FAvaSceneStateRCControllerMapping>& GetControllerMappings();
+
+	/** Gets the instanced property bag from ValuesHandle. Expects only one instance to be available */
+	FInstancedPropertyBag& GetInstancedPropertyBag();
+
+	/** Handle to the FAvaSceneStateRCTaskInstance holding the values and mappings properties */
+	TSharedRef<IPropertyHandle> StructHandle;
+
+	/** Handle to the controller values instanced property bag in FAvaSceneStateRCTaskInstance */
+	TSharedRef<IPropertyHandle> ValuesHandle;
+
+	/** Handle to the controller mappings array in FAvaSceneStateRCTaskInstance */
+	TSharedRef<IPropertyHandle> MappingsHandle;
+
+	/** Property utilities to use for refreshing the details view */
+	TSharedPtr<IPropertyUtilities> PropertyUtilities;
+
+	/** Binding id to the values Instanced property bag */
+	FGuid ValuesId;
+
+	/** Delegate called when needing to rebuild the children items */
+	FSimpleDelegate OnRebuildChildren;
+
+	/** Flag to ensure Initialize() only gets called once */
+	bool bInitialized = false;
+};
