Motion Design Scene State: Auto-register derived classes of Scene State Blueprint in the Scene State BP Editor Module, rather than having each plugin register themselves in their own plugin editor module. In -game editor modules are not loaded by default so blueprints call "LoadModulesRequiredForCompilation" to have the (usually editor-only) modules housing the kismet compiler loaded and the compilers registered. To avoid this extra work for derived blueprints, this is will now be done automatically by the Scene State BP Editor module.

diff --git a/Engine/Plugins/Experimental/AvalancheSceneState/Source/AvalancheSceneStateEditor/Private/AvaSceneStateEditorModule.cpp b/Engine/Plugins/Experimental/AvalancheSceneState/Source/AvalancheSceneStateEditor/Private/AvaSceneStateEditorModule.cpp
--- a/Engine/Plugins/Experimental/AvalancheSceneState/Source/AvalancheSceneStateEditor/Private/AvaSceneStateEditorModule.cpp
+++ b/Engine/Plugins/Experimental/AvalancheSceneState/Source/AvalancheSceneStateEditor/Private/AvaSceneStateEditorModule.cpp
@@ -6,7 +6,6 @@
 #include "AvaSceneStateEditorStyle.h"
 #include "AvaSceneStateExtension.h"
 #include "DetailsView/AvaSceneStateRCTaskDetails.h"
-#include "ISceneStateBlueprintEditorModule.h"
 #include "Modules/ModuleManager.h"
 #include "PropertyEditorModule.h"
 #include "RemoteControl/AvaSceneStateRCEventBehaviorNode.h"
@@ -18,8 +17,6 @@ void FAvaSceneStateEditorModule::StartupModule()
 {
 	UE::AvaSceneState::Editor::FEditorStyle::Get();
 
-	UE::SceneState::Editor::IBlueprintEditorModule::Get().RegisterCompiler(UAvaSceneStateBlueprint::StaticClass());
-
 	OnEditorBuildHandle = FAvaEditorBuilder::OnEditorBuild.AddLambda(
 		[](FAvaEditorBuilder& InBuilder)
 		{
diff --git a/Engine/Plugins/Experimental/SceneState/Source/SceneStateBlueprintEditor/Private/SceneStateBlueprintEditorModule.cpp b/Engine/Plugins/Experimental/SceneState/Source/SceneStateBlueprintEditor/Private/SceneStateBlueprintEditorModule.cpp
--- a/Engine/Plugins/Experimental/SceneState/Source/SceneStateBlueprintEditor/Private/SceneStateBlueprintEditorModule.cpp
+++ b/Engine/Plugins/Experimental/SceneState/Source/SceneStateBlueprintEditor/Private/SceneStateBlueprintEditorModule.cpp
@@ -133,7 +133,13 @@ void FBlueprintEditorModule::RegisterCompiler()
 	KismetCompilerModule.OverrideBPTypeForClass(USceneStateObject::StaticClass(), USceneStateBlueprint::StaticClass());
 	KismetCompilerModule.OverrideBPTypeForClass(USceneStateBlueprintableTask::StaticClass(), USceneStateTaskBlueprint::StaticClass());
 
-	FBlueprintEditorModule::RegisterCompiler(USceneStateBlueprint::StaticClass());
+	for (UClass* Class : TObjectRange<UClass>())
+	{
+		if (Class && Class->IsChildOf<USceneStateBlueprint>())
+		{
+			FBlueprintEditorModule::RegisterCompiler(Class);
+		}
+	}
 }
 
 void FBlueprintEditorModule::RegisterDefaultEvents()
