Motion Design: fixed scene state blueprints not  updating their generated class on world rename, causing issues with unable to saving levels, potentially leading to data loss if user attempts to continue on the action.

diff --git a/Engine/Plugins/Experimental/AvalancheSceneState/Source/AvalancheSceneState/Private/AvaSceneStateActor.cpp b/Engine/Plugins/Experimental/AvalancheSceneState/Source/AvalancheSceneState/Private/AvaSceneStateActor.cpp
--- a/Engine/Plugins/Experimental/AvalancheSceneState/Source/AvalancheSceneState/Private/AvaSceneStateActor.cpp
+++ b/Engine/Plugins/Experimental/AvalancheSceneState/Source/AvalancheSceneState/Private/AvaSceneStateActor.cpp
@@ -15,6 +15,7 @@ AAvaSceneStateActor::AAvaSceneStateActor(const FObjectInitializer& InObjectIniti
 {
 #if WITH_EDITOR
 	FWorldDelegates::OnWorldCleanup.AddUObject(this, &AAvaSceneStateActor::OnWorldCleanup);
+	FWorldDelegates::OnPostWorldRename.AddUObject(this, &AAvaSceneStateActor::OnWorldRenamed);
 #endif
 }
 
@@ -50,6 +51,7 @@ void AAvaSceneStateActor::BeginDestroy()
 {
 	Super::BeginDestroy();
 	FWorldDelegates::OnWorldCleanup.RemoveAll(this);
+	FWorldDelegates::OnPostWorldRename.RemoveAll(this);
 }
 
 void AAvaSceneStateActor::UpdateSceneStateClass()
@@ -85,6 +87,16 @@ void AAvaSceneStateActor::OnSceneStateRecompiled(UBlueprint* InCompiledBlueprint
 	UpdateSceneStateClass();
 }
 
+void AAvaSceneStateActor::OnWorldRenamed(UWorld* InNewWorld)
+{
+	// Rename Blueprints to the same Outer/Name so that the latest information is propagated to renaming the generated class
+	// Not doing this can result in the generated class properties to be bound to the old name
+	if (SceneStateBlueprint)
+	{
+		SceneStateBlueprint->Rename(*SceneStateBlueprint->GetName(), SceneStateBlueprint->GetOuter(), REN_DontCreateRedirectors);
+	}
+}
+
 void AAvaSceneStateActor::OnWorldCleanup(UWorld* InWorld, bool bInSessionEnded, bool bInCleanupResources)
 {
 	if (!bInCleanupResources || !SceneStateBlueprint)
diff --git a/Engine/Plugins/Experimental/AvalancheSceneState/Source/AvalancheSceneState/Public/AvaSceneStateActor.h b/Engine/Plugins/Experimental/AvalancheSceneState/Source/AvalancheSceneState/Public/AvaSceneStateActor.h
--- a/Engine/Plugins/Experimental/AvalancheSceneState/Source/AvalancheSceneState/Public/AvaSceneStateActor.h
+++ b/Engine/Plugins/Experimental/AvalancheSceneState/Source/AvalancheSceneState/Public/AvaSceneStateActor.h
@@ -40,6 +40,7 @@ protected:
 
 	void OnSceneStateRecompiled(UBlueprint* InCompiledBlueprint);
 
+	void OnWorldRenamed(UWorld* InNewWorld);
 	void OnWorldCleanup(UWorld* InWorld, bool bInSessionEnded, bool bInCleanupResources);
 #endif
 
