Motion Design: simplified promotion copy logic

diff --git a/Engine/Plugins/Experimental/AvalancheDataLink/Source/AvalancheDataLink/Private/AvaDataLinkInstance.cpp b/Engine/Plugins/Experimental/AvalancheDataLink/Source/AvalancheDataLink/Private/AvaDataLinkInstance.cpp
--- a/Engine/Plugins/Experimental/AvalancheDataLink/Source/AvalancheDataLink/Private/AvaDataLinkInstance.cpp
+++ b/Engine/Plugins/Experimental/AvalancheDataLink/Source/AvalancheDataLink/Private/AvaDataLinkInstance.cpp
@@ -38,35 +38,22 @@ namespace UE::AvaDataLink::Private
 	template<typename InSourcePropertyType>
 	class FPromotionCopy
 	{
-		using FSourceCppType = typename InSourcePropertyType::TCppType;
-
 		template<typename InTargetPropertyType>
 		bool CopySingle(const FCopyInfo& InCopyInfo, const InSourcePropertyType& InSourceProperty)
 		{
-			using FTargetCppType = typename InTargetPropertyType::TCppType;
+			using FSourceCppType = InSourcePropertyType::TCppType;
+			using FTargetCppType = InTargetPropertyType::TCppType;
 
 			static_assert(std::is_convertible_v<FSourceCppType, FTargetCppType>);
 
 			if (InTargetPropertyType* TargetProperty = CastField<InTargetPropertyType>(InCopyInfo.TargetProperty))
 			{
-				TargetProperty->SetPropertyValue(InCopyInfo.TargetMemory, static_cast<FTargetCppType>(InSourceProperty.GetPropertyValue(InCopyInfo.SourceMemory)));
+				TargetProperty->SetPropertyValue(InCopyInfo.TargetMemory, InSourceProperty.GetPropertyValue(InCopyInfo.SourceMemory));
 				return true;
 			}
 			return false;
 		}
 
-		template<typename>
-		bool CopyInternal(const FCopyInfo& InCopyInfo, const InSourcePropertyType& InSourceProperty)
-		{
-			return false;
-		}
-
-		template<typename, typename InTargetPropertyType, typename ...InOtherTargetPropertyTypes>
-		bool CopyInternal(const FCopyInfo& InCopyInfo, const InSourcePropertyType& InSourceProperty)
-		{
-			return this->CopySingle<InTargetPropertyType>(InCopyInfo, InSourceProperty) || this->CopyInternal<void, InOtherTargetPropertyTypes...>(InCopyInfo, InSourceProperty);
-		}
-
 	public:
 		/** Tries promoting from a source property type to any of the given target property types */
 		template<typename ...InTargetPropertyTypes>
@@ -74,7 +61,7 @@ namespace UE::AvaDataLink::Private
 		{
 			if (const InSourcePropertyType* SourceProperty = CastField<InSourcePropertyType>(InCopyInfo.SourceProperty))
 			{
-				return this->CopyInternal<void, InTargetPropertyTypes...>(InCopyInfo, *SourceProperty);
+				return (this->CopySingle<InTargetPropertyTypes>(InCopyInfo, *SourceProperty) || ...);
 			}
 			return false;
 		}
diff --git a/Engine/Plugins/Experimental/AvalancheSceneState/Source/AvalancheSceneState/Private/RemoteControl/AvaSceneStateRCTask.cpp b/Engine/Plugins/Experimental/AvalancheSceneState/Source/AvalancheSceneState/Private/RemoteControl/AvaSceneStateRCTask.cpp
--- a/Engine/Plugins/Experimental/AvalancheSceneState/Source/AvalancheSceneState/Private/RemoteControl/AvaSceneStateRCTask.cpp
+++ b/Engine/Plugins/Experimental/AvalancheSceneState/Source/AvalancheSceneState/Private/RemoteControl/AvaSceneStateRCTask.cpp
@@ -30,35 +30,22 @@ namespace UE::AvaSceneState::Private
 	template<typename InSourcePropertyType>
 	class FPromotionCopy
 	{
-		using FSourceCppType = typename InSourcePropertyType::TCppType;
-
 		template<typename InTargetPropertyType>
 		bool CopySingle(const FCopyInfo& InCopyInfo, const InSourcePropertyType& InSourceProperty)
 		{
-			using FTargetCppType = typename InTargetPropertyType::TCppType;
+			using FSourceCppType = InSourcePropertyType::TCppType;
+			using FTargetCppType = InTargetPropertyType::TCppType;
 
 			static_assert(std::is_convertible_v<FSourceCppType, FTargetCppType>);
 
 			if (InTargetPropertyType* TargetProperty = CastField<InTargetPropertyType>(InCopyInfo.TargetProperty))
 			{
-				TargetProperty->SetPropertyValue(InCopyInfo.TargetMemory, static_cast<FTargetCppType>(InSourceProperty.GetPropertyValue(InCopyInfo.SourceMemory)));
+				TargetProperty->SetPropertyValue(InCopyInfo.TargetMemory, InSourceProperty.GetPropertyValue(InCopyInfo.SourceMemory));
 				return true;
 			}
 			return false;
 		}
 
-		template<typename>
-		bool CopyInternal(const FCopyInfo& InCopyInfo, const InSourcePropertyType& InSourceProperty)
-		{
-			return false;
-		}
-
-		template<typename, typename InTargetPropertyType, typename ...InOtherTargetPropertyTypes>
-		bool CopyInternal(const FCopyInfo& InCopyInfo, const InSourcePropertyType& InSourceProperty)
-		{
-			return this->CopySingle<InTargetPropertyType>(InCopyInfo, InSourceProperty) || this->CopyInternal<void, InOtherTargetPropertyTypes...>(InCopyInfo, InSourceProperty);
-		}
-
 	public:
 		/** Tries promoting from a source property type to any of the given target property types */
 		template<typename ...InTargetPropertyTypes>
@@ -66,7 +53,7 @@ namespace UE::AvaSceneState::Private
 		{
 			if (const InSourcePropertyType* SourceProperty = CastField<InSourcePropertyType>(InCopyInfo.SourceProperty))
 			{
-				return this->CopyInternal<void, InTargetPropertyTypes...>(InCopyInfo, *SourceProperty);
+				return (this->CopySingle<InTargetPropertyTypes>(InCopyInfo, *SourceProperty) || ...);
 			}
 			return false;
 		}
