Motion Design Scene State: allow motion design scene state actor to be deleted via button

diff --git a/Engine/Plugins/Experimental/AvalancheSceneState/Source/AvalancheSceneState/Private/AvaSceneStateActor.cpp b/Engine/Plugins/Experimental/AvalancheSceneState/Source/AvalancheSceneState/Private/AvaSceneStateActor.cpp
--- a/Engine/Plugins/Experimental/AvalancheSceneState/Source/AvalancheSceneState/Private/AvaSceneStateActor.cpp
+++ b/Engine/Plugins/Experimental/AvalancheSceneState/Source/AvalancheSceneState/Private/AvaSceneStateActor.cpp
@@ -110,13 +110,21 @@ void AAvaSceneStateActor::OnWorldCleanup(UWorld* InWorld, bool bInSessionEnded,
 		return;
 	}
 
-	const FName ObjectName = MakeUniqueObjectName(GetTransientPackage()
-		, SceneStateBlueprint->GetClass()
-		, FName(*FString::Printf(TEXT("%s_Trashed"), *SceneStateBlueprint->GetName())));
+	CleanupSceneState();
+}
+
+void AAvaSceneStateActor::CleanupSceneState()
+{
+	if (SceneStateBlueprint)
+	{
+		const FName ObjectName = MakeUniqueObjectName(GetTransientPackage()
+			, SceneStateBlueprint->GetClass()
+			, FName(*FString::Printf(TEXT("%s_Trashed"), *SceneStateBlueprint->GetName())));
 
-	SceneStateBlueprint->Rename(*ObjectName.ToString()
-		, GetTransientPackage()
-		, REN_DontCreateRedirectors | REN_NonTransactional | REN_DoNotDirty);
+		SceneStateBlueprint->Rename(*ObjectName.ToString()
+			, GetTransientPackage()
+			, REN_DontCreateRedirectors | REN_NonTransactional | REN_DoNotDirty);
+	}
 
 	SetSceneStateBlueprint(nullptr);
 	SetSceneStateClass(nullptr);
diff --git a/Engine/Plugins/Experimental/AvalancheSceneState/Source/AvalancheSceneState/Public/AvaSceneStateActor.h b/Engine/Plugins/Experimental/AvalancheSceneState/Source/AvalancheSceneState/Public/AvaSceneStateActor.h
--- a/Engine/Plugins/Experimental/AvalancheSceneState/Source/AvalancheSceneState/Public/AvaSceneStateActor.h
+++ b/Engine/Plugins/Experimental/AvalancheSceneState/Source/AvalancheSceneState/Public/AvaSceneStateActor.h
@@ -6,6 +6,8 @@
 #include "SceneStateActor.h"
 #include "AvaSceneStateActor.generated.h"
 
+#define UE_API AVALANCHESCENESTATE_API
+
 #if WITH_EDITOR
 class UBlueprint;
 class USceneStateBlueprint;
@@ -34,17 +36,27 @@ protected:
 	virtual void BeginDestroy() override;
 	//~ End UObject
 
-	AVALANCHESCENESTATE_API void UpdateSceneStateClass();
+	/** Updates the scene state class to the latest generated class from the scene state blueprint if valid */
+	UE_API void UpdateSceneStateClass();
 
-	AVALANCHESCENESTATE_API void SetSceneStateBlueprint(USceneStateBlueprint* InSceneStateBlueprint);
+	UE_API void SetSceneStateBlueprint(USceneStateBlueprint* InSceneStateBlueprint);
 
+	/** Called when the scene state blueprint is compiled */
 	void OnSceneStateRecompiled(UBlueprint* InCompiledBlueprint);
 
+	/** Called when the world is renamed */
 	void OnWorldRenamed(UWorld* InNewWorld);
+
+	/** Called on World destruction / change to clean up scene state blueprint resources */
 	void OnWorldCleanup(UWorld* InWorld, bool bInSessionEnded, bool bInCleanupResources);
 #endif
 
 private:
+#if WITH_EDITOR
+	/** Cleans up scene state resources */
+	UE_API void CleanupSceneState();
+#endif
+
 #if WITH_EDITORONLY_DATA
 	UPROPERTY()
 	TObjectPtr<UBlueprint> SceneStateBlueprint;
@@ -52,3 +64,5 @@ private:
 
 	friend class FAvaSceneStateExtension;
 };
+
+#undef UE_API
diff --git a/Engine/Plugins/Experimental/AvalancheSceneState/Source/AvalancheSceneStateEditor/Private/AvaSceneStateExtension.cpp b/Engine/Plugins/Experimental/AvalancheSceneState/Source/AvalancheSceneStateEditor/Private/AvaSceneStateExtension.cpp
--- a/Engine/Plugins/Experimental/AvalancheSceneState/Source/AvalancheSceneStateEditor/Private/AvaSceneStateExtension.cpp
+++ b/Engine/Plugins/Experimental/AvalancheSceneState/Source/AvalancheSceneStateEditor/Private/AvaSceneStateExtension.cpp
@@ -7,9 +7,13 @@
 #include "Editor.h"
 #include "Engine/World.h"
 #include "EngineUtils.h"
+#include "Misc/MessageDialog.h"
 #include "SceneStateActor.h"
 #include "SceneStateBlueprintFactory.h"
+#include "SceneStateComponent.h"
+#include "SceneStateComponentPlayer.h"
 #include "SceneStateObject.h"
+#include "ScopedTransaction.h"
 #include "Styling/SlateIconFinder.h"
 #include "Subsystems/AssetEditorSubsystem.h"
 #include "ToolMenu.h"
@@ -22,13 +26,23 @@ void FAvaSceneStateExtension::ExtendToolbarMenu(UToolMenu& InMenu)
 {
 	FToolMenuSection& Section = InMenu.FindOrAddSection(DefaultSectionName);
 
+	FSlateIcon SceneStateIcon = FSlateIconFinder::FindCustomIconForClass(UAvaSceneStateBlueprint::StaticClass(), TEXT("ClassThumbnail"));
+
 	FToolMenuEntry& Entry = Section.AddEntry(FToolMenuEntry::InitToolBarButton(TEXT("SceneStateButton")
 		, FExecuteAction::CreateSP(this, &FAvaSceneStateExtension::OpenSceneStateBlueprintEditor)
 		, LOCTEXT("SceneStateLabel", "Scene State")
 		, LOCTEXT("SceneStateTooltip", "Opens the Scene State Editor for the given Motion Design Scene")
-		, FSlateIconFinder::FindCustomIconForClass(UAvaSceneStateBlueprint::StaticClass(), TEXT("ClassThumbnail"))));
+		, SceneStateIcon));
 
 	Entry.StyleNameOverride = TEXT("CalloutToolbar");
+
+	Section.AddEntry(FToolMenuEntry::InitComboButton(TEXT("SceneStateComboButton")
+		, FUIAction()
+		, FNewToolMenuDelegate::CreateSP(this, &FAvaSceneStateExtension::GenerateSceneStateOptions)
+		, LOCTEXT("SceneStateOptionsLabel", "Scene State Options")
+		, LOCTEXT("SceneStateOptionsTooltip", "Scene State Options")
+		, SceneStateIcon
+		, /*bSimpleComboBox*/true));
 }
 
 void FAvaSceneStateExtension::Cleanup()
@@ -54,7 +68,32 @@ void FAvaSceneStateExtension::Cleanup()
 	}
 }
 
-AAvaSceneStateActor* FAvaSceneStateExtension::FindOrSpawnActor() const
+void FAvaSceneStateExtension::GenerateSceneStateOptions(UToolMenu* InMenu)
+{
+	FToolMenuSection& GeneralSection = InMenu->FindOrAddSection(TEXT("GeneralSection"), LOCTEXT("GeneralSectionLabel", "General"));
+
+	GeneralSection.AddMenuEntry(TEXT("DeleteSceneState")
+		, LOCTEXT("DeleteSceneStateLabel", "Delete Scene State")
+		, LOCTEXT("DeleteSceneStateTooltip", "Deletes the existing embedded Scene State")
+		, TAttribute<FSlateIcon>()
+		, FUIAction(FExecuteAction::CreateSP(this, &FAvaSceneStateExtension::DeleteSceneStateActor)
+			, FCanExecuteAction::CreateSP(this, &FAvaSceneStateExtension::CanDeleteSceneStateActor))
+		, EUserInterfaceActionType::Button);
+}
+
+AAvaSceneStateActor* FAvaSceneStateExtension::FindSceneStateActor() const
+{
+	if (UWorld* const World = GetWorld())
+	{
+		for (AAvaSceneStateActor* SceneStateActor : TActorRange<AAvaSceneStateActor>(World))
+		{
+			return SceneStateActor;
+		}
+	}
+	return nullptr;
+}
+
+AAvaSceneStateActor* FAvaSceneStateExtension::FindOrSpawnSceneStateActor() const
 {
 	UWorld* const World = GetWorld();
 	if (!World)
@@ -93,9 +132,42 @@ USceneStateBlueprint* FAvaSceneStateExtension::CreateSceneStateBlueprint(AAvaSce
 	return NewSceneStateBlueprint;
 }
 
+bool FAvaSceneStateExtension::CanDeleteSceneStateActor() const
+{
+	return !!FindSceneStateActor();
+}
+
+void FAvaSceneStateExtension::DeleteSceneStateActor()
+{
+	const EAppReturnType::Type Response = FMessageDialog::Open(EAppMsgType::YesNo
+		, LOCTEXT("DeleteSceneStateMessage", "Do you wish to delete this Scene State Actor?\nThe embedded scene state blueprint will be destroyed with it.")
+		, LOCTEXT("DeleteSceneStateTitle", "Delete Scene State"));
+
+	if (Response != EAppReturnType::Yes)
+	{
+		return;
+	}
+
+	AAvaSceneStateActor* SceneStateActor = FindSceneStateActor();
+	if (!SceneStateActor)
+	{
+		return;
+	}
+
+	if (SceneStateActor->SceneStateBlueprint)
+	{
+		UAssetEditorSubsystem* AssetEditorSubsystem = GEditor->GetEditorSubsystem<UAssetEditorSubsystem>();
+		check(AssetEditorSubsystem);
+		AssetEditorSubsystem->CloseAllEditorsForAsset(SceneStateActor->SceneStateBlueprint);
+	}
+
+	SceneStateActor->CleanupSceneState();
+	SceneStateActor->Destroy(/*bNetForce*/false, /*bShouldModifyLevel*/false);
+}
+
 void FAvaSceneStateExtension::OpenSceneStateBlueprintEditor()
 {
-	AAvaSceneStateActor* SceneStateActor = FindOrSpawnActor();
+	AAvaSceneStateActor* SceneStateActor = FindOrSpawnSceneStateActor();
 	if (!SceneStateActor)
 	{
 		return;
diff --git a/Engine/Plugins/Experimental/AvalancheSceneState/Source/AvalancheSceneStateEditor/Private/AvaSceneStateExtension.h b/Engine/Plugins/Experimental/AvalancheSceneState/Source/AvalancheSceneStateEditor/Private/AvaSceneStateExtension.h
--- a/Engine/Plugins/Experimental/AvalancheSceneState/Source/AvalancheSceneStateEditor/Private/AvaSceneStateExtension.h
+++ b/Engine/Plugins/Experimental/AvalancheSceneState/Source/AvalancheSceneStateEditor/Private/AvaSceneStateExtension.h
@@ -18,9 +18,23 @@ public:
 	//~ End IAvaEditorExtension
 
 private:
-	AAvaSceneStateActor* FindOrSpawnActor() const;
+	void GenerateSceneStateOptions(UToolMenu* InMenu);
 
+	/** Finds an existing scene state actor in the current world */
+	AAvaSceneStateActor* FindSceneStateActor() const;
+
+	/** Finds or spawns a scene state actor in the current world */
+	AAvaSceneStateActor* FindOrSpawnSceneStateActor() const;
+
+	/** Creates a new scene state blueprint and assigns it to the given scene state actor */
 	USceneStateBlueprint* CreateSceneStateBlueprint(AAvaSceneStateActor* InSceneStateActor) const;
 
+	/** Returns whether there's an existing scene state actor that can be destroyed */
+	bool CanDeleteSceneStateActor() const;
+
+	/** Destroys the existing scene state actor in the current world */
+	void DeleteSceneStateActor();
+
+	/** Opens the existing scene state blueprint in the current world*/
 	void OpenSceneStateBlueprintEditor();
 };
diff --git a/Engine/Plugins/Experimental/SceneState/Source/SceneStateGameplay/Public/SceneStateActor.h b/Engine/Plugins/Experimental/SceneState/Source/SceneStateGameplay/Public/SceneStateActor.h
--- a/Engine/Plugins/Experimental/SceneState/Source/SceneStateGameplay/Public/SceneStateActor.h
+++ b/Engine/Plugins/Experimental/SceneState/Source/SceneStateGameplay/Public/SceneStateActor.h
@@ -5,6 +5,8 @@
 #include "GameFramework/Actor.h"
 #include "SceneStateActor.generated.h"
 
+#define UE_API SCENESTATEGAMEPLAY_API
+
 class USceneStateComponent;
 class USceneStateObject;
 
@@ -14,18 +16,25 @@ class ASceneStateActor : public AActor
 	GENERATED_BODY()
 
 public:
-	SCENESTATEGAMEPLAY_API static const FLazyName SceneStateComponentName;
+	UE_API static const FLazyName SceneStateComponentName;
+
+	UE_API ASceneStateActor(const FObjectInitializer& InObjectInitializer);
 
-	SCENESTATEGAMEPLAY_API ASceneStateActor(const FObjectInitializer& InObjectInitializer);
+	USceneStateComponent* GetSceneStateComponent() const
+	{
+		return SceneStateComponent;
+	}
 
-	SCENESTATEGAMEPLAY_API void SetSceneStateClass(TSubclassOf<USceneStateObject> InSceneStateClass);
+	UE_API void SetSceneStateClass(TSubclassOf<USceneStateObject> InSceneStateClass);
 
-	SCENESTATEGAMEPLAY_API TSubclassOf<USceneStateObject> GetSceneStateClass() const;
+	UE_API TSubclassOf<USceneStateObject> GetSceneStateClass() const;
 
 	UFUNCTION(BlueprintCallable, Category = "Scene State")
-	SCENESTATEGAMEPLAY_API USceneStateObject* GetSceneState() const;
+	UE_API USceneStateObject* GetSceneState() const;
 
 private:
 	UPROPERTY(VisibleAnywhere, BlueprintReadOnly, Category="Scene State", meta=(AllowPrivateAccess="true"))
 	TObjectPtr<USceneStateComponent> SceneStateComponent;
 };
+
+#undef UE_API
diff --git a/Engine/Plugins/Experimental/SceneState/Source/SceneStateGameplay/Public/SceneStateComponent.h b/Engine/Plugins/Experimental/SceneState/Source/SceneStateGameplay/Public/SceneStateComponent.h
--- a/Engine/Plugins/Experimental/SceneState/Source/SceneStateGameplay/Public/SceneStateComponent.h
+++ b/Engine/Plugins/Experimental/SceneState/Source/SceneStateGameplay/Public/SceneStateComponent.h
@@ -5,6 +5,8 @@
 #include "Components/ActorComponent.h"
 #include "SceneStateComponent.generated.h"
 
+#define UE_API SCENESTATEGAMEPLAY_API
+
 class USceneStateObject;
 class USceneStateComponentPlayer;
 struct FSceneStateComponentInstanceData;
@@ -15,34 +17,36 @@ class USceneStateComponent : public UActorComponent
 	GENERATED_BODY()
 
 public:
-	SCENESTATEGAMEPLAY_API static const FLazyName SceneStatePlayerName;
+	UE_API static const FLazyName SceneStatePlayerName;
 
-	SCENESTATEGAMEPLAY_API USceneStateComponent(const FObjectInitializer& InObjectInitializer);
+	UE_API USceneStateComponent(const FObjectInitializer& InObjectInitializer);
 
 	USceneStateComponentPlayer* GetSceneStatePlayer() const
 	{
 		return SceneStatePlayer;
 	}
 
-	SCENESTATEGAMEPLAY_API TSubclassOf<USceneStateObject> GetSceneStateClass() const;
+	UE_API TSubclassOf<USceneStateObject> GetSceneStateClass() const;
 
-	SCENESTATEGAMEPLAY_API void SetSceneStateClass(TSubclassOf<USceneStateObject> InSceneStateClass);
+	UE_API void SetSceneStateClass(TSubclassOf<USceneStateObject> InSceneStateClass);
 
 	UFUNCTION(BlueprintCallable, Category = "Scene State")
-	SCENESTATEGAMEPLAY_API USceneStateObject* GetSceneState() const;
+	UE_API USceneStateObject* GetSceneState() const;
 
 	void ApplyComponentInstanceData(FSceneStateComponentInstanceData* InComponentInstanceData);
 
 	//~ Begin UActorComponent
-	SCENESTATEGAMEPLAY_API virtual void InitializeComponent() override;
-	SCENESTATEGAMEPLAY_API virtual void OnRegister() override;
-	SCENESTATEGAMEPLAY_API virtual void BeginPlay() override;
-	SCENESTATEGAMEPLAY_API virtual void TickComponent(float InDeltaTime, ELevelTick InTickType, FActorComponentTickFunction* InThisTickFunction) override;
-	SCENESTATEGAMEPLAY_API virtual void EndPlay(const EEndPlayReason::Type InEndPlayReason) override;
-	SCENESTATEGAMEPLAY_API virtual TStructOnScope<FActorComponentInstanceData> GetComponentInstanceData() const override;
+	UE_API virtual void InitializeComponent() override;
+	UE_API virtual void OnRegister() override;
+	UE_API virtual void BeginPlay() override;
+	UE_API virtual void TickComponent(float InDeltaTime, ELevelTick InTickType, FActorComponentTickFunction* InThisTickFunction) override;
+	UE_API virtual void EndPlay(const EEndPlayReason::Type InEndPlayReason) override;
+	UE_API virtual TStructOnScope<FActorComponentInstanceData> GetComponentInstanceData() const override;
 	//~ End UActorComponent
 
 private:
 	UPROPERTY(VisibleAnywhere, Instanced, Category="Scene State", meta=(ShowInnerProperties))
 	TObjectPtr<USceneStateComponentPlayer> SceneStatePlayer;
 };
+
+#undef UE_API
