Motion Design: fixed crash when deleting scene state actor (and its embedded bp) while the actor is still selected

diff --git a/Engine/Plugins/Experimental/AvalancheSceneState/Source/AvalancheSceneStateEditor/Private/AvaSceneStateExtension.cpp b/Engine/Plugins/Experimental/AvalancheSceneState/Source/AvalancheSceneStateEditor/Private/AvaSceneStateExtension.cpp
--- a/Engine/Plugins/Experimental/AvalancheSceneState/Source/AvalancheSceneStateEditor/Private/AvaSceneStateExtension.cpp
+++ b/Engine/Plugins/Experimental/AvalancheSceneState/Source/AvalancheSceneStateEditor/Private/AvaSceneStateExtension.cpp
@@ -5,6 +5,8 @@
 #include "AvaSceneStateBlueprint.h"
 #include "BlueprintActionDatabase.h"
 #include "Editor.h"
+#include "Editor/UnrealEdEngine.h"
+#include "EditorModeManager.h"
 #include "Engine/World.h"
 #include "EngineUtils.h"
 #include "Misc/MessageDialog.h"
@@ -14,11 +16,13 @@
 #include "SceneStateComponentPlayer.h"
 #include "SceneStateObject.h"
 #include "ScopedTransaction.h"
+#include "Selection.h"
 #include "Styling/SlateIconFinder.h"
 #include "Subsystems/AssetEditorSubsystem.h"
 #include "ToolMenu.h"
 #include "ToolMenuEntry.h"
 #include "ToolMenuSection.h"
+#include "UnrealEdGlobals.h"
 
 #define LOCTEXT_NAMESPACE "AvaSceneStateExtension"
 
@@ -162,6 +166,22 @@ void FAvaSceneStateExtension::DeleteSceneStateActor()
 	}
 
 	SceneStateActor->CleanupSceneState();
+
+	if (const FEditorModeTools* EditorModeTools = GetEditorModeTools())
+	{
+		if (USelection* SelectedComponents = EditorModeTools->GetSelectedComponents())
+		{
+			if (USceneStateComponent* SceneStateComponent = SceneStateActor->GetSceneStateComponent())
+			{
+				SelectedComponents->Deselect(SceneStateComponent);	
+			}
+		}
+		if (USelection* SelectedActors = EditorModeTools->GetSelectedActors())
+		{
+			SelectedActors->Deselect(SceneStateActor);
+		}
+	}
+
 	SceneStateActor->Destroy(/*bNetForce*/false, /*bShouldModifyLevel*/false);
 }
 
