Motion Design: set the scene state and data link plugins to beta

diff --git a/Engine/Plugins/VirtualProduction/AvalancheDataLink/AvalancheDataLink.uplugin b/Engine/Plugins/VirtualProduction/AvalancheDataLink/AvalancheDataLink.uplugin
--- a/Engine/Plugins/VirtualProduction/AvalancheDataLink/AvalancheDataLink.uplugin
+++ b/Engine/Plugins/VirtualProduction/AvalancheDataLink/AvalancheDataLink.uplugin
@@ -10,8 +10,8 @@
 	"DocsURL": "",
 	"MarketplaceURL": "",
 	"CanContainContent": true,
-	"IsBetaVersion": false,
-	"IsExperimentalVersion": true,
+	"IsBetaVersion": true,
+	"IsExperimentalVersion": false,
 	"Installed": false,
 	"Modules": [
 		{
diff --git a/Engine/Plugins/VirtualProduction/AvalancheSceneState/AvalancheSceneState.uplugin b/Engine/Plugins/VirtualProduction/AvalancheSceneState/AvalancheSceneState.uplugin
--- a/Engine/Plugins/VirtualProduction/AvalancheSceneState/AvalancheSceneState.uplugin
+++ b/Engine/Plugins/VirtualProduction/AvalancheSceneState/AvalancheSceneState.uplugin
@@ -10,8 +10,8 @@
 	"DocsURL": "",
 	"MarketplaceURL": "",
 	"CanContainContent": true,
-	"IsBetaVersion": false,
-	"IsExperimentalVersion": true,
+	"IsBetaVersion": true,
+	"IsExperimentalVersion": false,
 	"Installed": false,
 	"Modules": [
 		{
diff --git a/Engine/Plugins/VirtualProduction/DataLink/DataLink.uplugin b/Engine/Plugins/VirtualProduction/DataLink/DataLink.uplugin
--- a/Engine/Plugins/VirtualProduction/DataLink/DataLink.uplugin
+++ b/Engine/Plugins/VirtualProduction/DataLink/DataLink.uplugin
@@ -10,8 +10,8 @@
 	"DocsURL": "",
 	"MarketplaceURL": "",
 	"CanContainContent": false,
-	"IsBetaVersion": false,
-	"IsExperimentalVersion": true,
+	"IsBetaVersion": true,
+	"IsExperimentalVersion": false,
 	"Installed": false,
 	"Modules": [
 		{
diff --git a/Engine/Plugins/VirtualProduction/DataLinkOAuth/DataLinkOAuth.uplugin b/Engine/Plugins/VirtualProduction/DataLinkOAuth/DataLinkOAuth.uplugin
--- a/Engine/Plugins/VirtualProduction/DataLinkOAuth/DataLinkOAuth.uplugin
+++ b/Engine/Plugins/VirtualProduction/DataLinkOAuth/DataLinkOAuth.uplugin
@@ -10,8 +10,8 @@
 	"DocsURL": "",
 	"MarketplaceURL": "",
 	"CanContainContent": true,
-	"IsBetaVersion": false,
-	"IsExperimentalVersion": true,
+	"IsBetaVersion": true,
+	"IsExperimentalVersion": false,
 	"Installed": false,
 	"Modules": [
 		{
diff --git a/Engine/Plugins/VirtualProduction/SceneState/SceneState.uplugin b/Engine/Plugins/VirtualProduction/SceneState/SceneState.uplugin
--- a/Engine/Plugins/VirtualProduction/SceneState/SceneState.uplugin
+++ b/Engine/Plugins/VirtualProduction/SceneState/SceneState.uplugin
@@ -10,8 +10,8 @@
 	"DocsURL": "",
 	"MarketplaceURL": "",
 	"CanContainContent": true,
-	"IsBetaVersion": false,
-	"IsExperimentalVersion": true,
+	"IsBetaVersion": true,
+	"IsExperimentalVersion": false,
 	"Installed": false,
 	"Modules": [
 		{
diff --git a/Engine/Plugins/VirtualProduction/SceneStateDataLink/SceneStateDataLink.uplugin b/Engine/Plugins/VirtualProduction/SceneStateDataLink/SceneStateDataLink.uplugin
--- a/Engine/Plugins/VirtualProduction/SceneStateDataLink/SceneStateDataLink.uplugin
+++ b/Engine/Plugins/VirtualProduction/SceneStateDataLink/SceneStateDataLink.uplugin
@@ -10,8 +10,8 @@
 	"DocsURL": "",
 	"MarketplaceURL": "",
 	"CanContainContent": true,
-	"IsBetaVersion": false,
-	"IsExperimentalVersion": true,
+	"IsBetaVersion": true,
+	"IsExperimentalVersion": false,
 	"Installed": false,
 	"Modules": [
 		{
